#TTP 332496 - EDITOR: Submit files dialog does not use the standard dialog button sizes.

#branch UE4

#proj Editor.SourceControlWindows

#summary The submit files dialog now follows the same style as the rest of the engine.

#change Change the buttons to use the established method of crating uniformly sized buttons.

diff --git a/Engine/Source/Editor/SourceControlWindows/Private/SSourceControlSubmit.cpp b/Engine/Source/Editor/SourceControlWindows/Private/SSourceControlSubmit.cpp
--- a/Engine/Source/Editor/SourceControlWindows/Private/SSourceControlSubmit.cpp
+++ b/Engine/Source/Editor/SourceControlWindows/Private/SSourceControlSubmit.cpp
@@ -227,31 +227,34 @@ public:
 					]
 				]
 				+SVerticalBox::Slot()
-				.FillHeight(1)
+				.AutoHeight()
+				.HAlign(HAlign_Right)
 				.VAlign(VAlign_Bottom)
-				.Padding(5,5,5,0)
+				.Padding(0.0f,0.0f,0.0f,5.0f)
 				[
-					SNew(SHorizontalBox)
-					+SHorizontalBox::Slot()
-					.HAlign(HAlign_Right)
-					.FillWidth(1)
+					SNew(SUniformGridPanel)
+					.SlotPadding(FEditorStyle::GetMargin("StandardDialog.SlotPadding"))
+					.MinDesiredSlotWidth(FEditorStyle::GetFloat("StandardDialog.MinDesiredSlotWidth"))
+					.MinDesiredSlotHeight(FEditorStyle::GetFloat("StandardDialog.MinDesiredSlotHeight"))
+					+SUniformGridPanel::Slot(0,0)
 					[
 						SNew(SButton)
+						.HAlign(HAlign_Center)
+						.ContentPadding(FEditorStyle::GetMargin("StandardDialog.ContentPadding"))
 						.IsEnabled(this, &SSourceControlSubmitWidget::IsOKEnabled)
 						.Text( NSLOCTEXT("SourceControl.SubmitPanel", "OKButton", "OK") )
 						.OnClicked(this, &SSourceControlSubmitWidget::OKClicked)
 					]
-					+SHorizontalBox::Slot()
-					.AutoWidth()
-					.HAlign(HAlign_Right)
+					+SUniformGridPanel::Slot(1,0)
 					[
 						SNew(SButton)
+						.HAlign(HAlign_Center)
+						.ContentPadding(FEditorStyle::GetMargin("StandardDialog.ContentPadding"))
 						.Text( NSLOCTEXT("SourceControl.SubmitPanel", "CancelButton", "Cancel") )
 						.OnClicked(this, &SSourceControlSubmitWidget::CancelClicked)
 					]
 				]
-			]					
-
+			]
 		]
 		;
 
