Submit message icon is no longer a 'stop' symbol

Changed the message log direct notify to a custom notification so the icon gets set to something less aggressive (message log uses the 'worst' message on the current log page for the icon).

TTP# 333620 - Editor: Source Control: Error icon displayed when submitting to source control

diff --git a/Engine/Source/Editor/SourceControlWindows/Private/SSourceControlSubmit.cpp b/Engine/Source/Editor/SourceControlWindows/Private/SSourceControlSubmit.cpp
--- a/Engine/Source/Editor/SourceControlWindows/Private/SSourceControlSubmit.cpp
+++ b/Engine/Source/Editor/SourceControlWindows/Private/SSourceControlSubmit.cpp
@@ -554,7 +554,14 @@ bool FSourceControlWindows::PromptForCheckin(const TArray<FString>& InPackageNam
 				if(bCheckInSuccess)
 				{
 					// report success with a notification
-					FMessageLog("SourceControl").Notify(CheckInOperation->GetSuccessMessage());
+					FNotificationInfo Info(CheckInOperation->GetSuccessMessage());
+					Info.ExpireDuration = 8.0f;
+					Info.HyperlinkText = LOCTEXT("SCC_Checkin_ShowLog", "Show Log");
+					Info.Hyperlink = FSimpleDelegate::CreateStatic([](){ FMessageLog("SourceControl").Open(EMessageSeverity::Info, true); });
+					FSlateNotificationManager::Get().AddNotification(Info);
+
+					// also add to the log
+					FMessageLog("SourceControl").Info(CheckInOperation->GetSuccessMessage());
 				}
 			}
 
