Updated the revert window to use Revert Unchanged Only correctly.
#TTP 334818 - Editor: Source Control: Revert files dialog checkboxes not behaving properly
#branch UE4

#proj Editor.SourceControlWindows

#change Changed the behaviour of the revert window as follows:
Previously:
Clicking revert unchanged only attempted to disable the changed files, all disabling did at this point was to make it seem as if the item was not selected even if it was. This meant that this check box did not fulfil its actual description

Now:
Checking revert unchanged only disables all of the file items as well as the checkbox to enable or disable all of the file items. It then enables the revert button and if pressed it will only send the unchanged files for reversion, this means it fulfills the buttons description.

#change Misc changes: Updated the buttons to use the standardized button layout, and changed "revert files" button to "revert" as this lead to a nicer layout. Added icons to show which of the selected files have been modified.

diff --git a/Engine/Source/Editor/SourceControlWindows/Private/SSourceControlRevert.cpp b/Engine/Source/Editor/SourceControlWindows/Private/SSourceControlRevert.cpp
--- a/Engine/Source/Editor/SourceControlWindows/Private/SSourceControlRevert.cpp
+++ b/Engine/Source/Editor/SourceControlWindows/Private/SSourceControlRevert.cpp
@@ -29,8 +29,8 @@ struct FRevertCheckBoxListViewItem
 	FRevertCheckBoxListViewItem( FString InText )
 	{
 		Text = InText;
-		IsEnabled = true;
 		IsSelected = false;
+		IsModified = false;
 	}
 
 	void OnCheckStateChanged( const ESlateCheckBoxState::Type NewCheckedState )
@@ -40,12 +40,16 @@ struct FRevertCheckBoxListViewItem
 
 	ESlateCheckBoxState::Type OnIsChecked() const
 	{
-		return (IsEnabled ? IsSelected : false) ? ESlateCheckBoxState::Checked : ESlateCheckBoxState::Unchecked;
+		return ( IsSelected ) ? ESlateCheckBoxState::Checked : ESlateCheckBoxState::Unchecked;
 	}
 
+	EVisibility OnGetModifiedStateVisibility() const
+	{
+		return (IsModified) ? EVisibility::Visible : EVisibility::Hidden;
+	}
 
-	bool IsEnabled;
 	bool IsSelected;
+	bool IsModified;
 	FString Text;
 };
 
@@ -111,6 +115,7 @@ public:
 					[
 						SNew(SCheckBox)
 						.OnCheckStateChanged(this, &SSourceControlRevertWidget::ColumnHeaderClicked)
+						.IsEnabled(this, &SSourceControlRevertWidget::OnGetItemsEnabled)
 						[
 							SNew(STextBlock)
 							.Text(NSLOCTEXT("SourceControl.Revert", "ListHeader", "File Name").ToString())
@@ -155,29 +160,37 @@ public:
 					.FillWidth(1)
 					.Padding(5)
 					[
-						SNew(SButton)
-						.OnClicked(this, &SSourceControlRevertWidget::OKClicked)
-						.IsEnabled(this, &SSourceControlRevertWidget::IsOKEnabled)
+						SNew(SUniformGridPanel)
+						.SlotPadding(FEditorStyle::GetMargin("StandardDialog.SlotPadding"))
+						.MinDesiredSlotWidth(FEditorStyle::GetFloat("StandardDialog.MinDesiredSlotWidth"))
+						.MinDesiredSlotHeight(FEditorStyle::GetFloat("StandardDialog.MinDesiredSlotHeight"))
+						+SUniformGridPanel::Slot(0,0)
 						[
-							SNew(STextBlock)
-							.Text(NSLOCTEXT("SourceControl.Revert", "RevertButton", "Revert Files").ToString())
+							SNew(SButton) 
+							.HAlign(HAlign_Center)
+							.ContentPadding(FEditorStyle::GetMargin("StandardDialog.ContentPadding"))
+							.OnClicked(this, &SSourceControlRevertWidget::OKClicked)
+							.IsEnabled(this, &SSourceControlRevertWidget::IsOKEnabled)
+							.Text(LOCTEXT("RevertButton", "Revert"))
+						]
+						+SUniformGridPanel::Slot(1,0)
+						[
+							SNew(SButton) 
+							.HAlign(HAlign_Center)
+							.ContentPadding(FEditorStyle::GetMargin("StandardDialog.ContentPadding"))
+							.OnClicked(this, &SSourceControlRevertWidget::CancelClicked)
+							.Text(LOCTEXT("CancelButton", "Cancel"))
 						]
-					]
-					+SHorizontalBox::Slot()
-					.AutoWidth()
-					.HAlign(HAlign_Right)
-					.Padding(0,5,5,5)
-					[
-						SNew(SButton)
-						.OnClicked(this, &SSourceControlRevertWidget::CancelClicked)
-						.Text(NSLOCTEXT("SourceControl.Revert", "CancelButton", "Cancel").ToString())
 					]
 				]
 			]
 		];
 
+		// update the modified state of all the files. 
+		UpdateSCCStatus();
+
 		DialogResult = ERevertResults::REVERT_CANCELED;
-		bAlreadyQueriedModifiedFiles = false;
+		bRevertUnchangedFilesOnly = false;
 	}
 	END_SLATE_FUNCTION_BUILD_OPTIMIZATION
 
@@ -188,11 +201,12 @@ public:
 	 */
 	void GetPackagesToRevert( TArray<FString>& OutPackagesToRevert )
 	{
-		for (int32 i=0; i<ListViewItemSource.Num(); i++)
+		for ( const auto& ListViewItem : ListViewItemSource )
 		{
-			if (ListViewItemSource[i]->IsSelected)
+			if ((bRevertUnchangedFilesOnly && !ListViewItem->IsModified) || 
+				(!bRevertUnchangedFilesOnly && ListViewItem->IsSelected))
 			{
-				OutPackagesToRevert.Add(ListViewItemSource[i]->Text);
+				OutPackagesToRevert.Add(ListViewItem->Text);
 			}
 		}
 	}
@@ -210,20 +224,34 @@ private:
 		TSharedPtr<SCheckBox> CheckBox;
 		TSharedRef<ITableRow> Row =
 			SNew(STableRow< TSharedPtr<FString> >, OwnerTable)
+			.IsEnabled(this, &SSourceControlRevertWidget::OnGetItemsEnabled)
 			[
-				SAssignNew(CheckBox, SCheckBox)
-				.OnCheckStateChanged(ListItemPtr.ToSharedRef(), &FRevertCheckBoxListViewItem::OnCheckStateChanged)
-				.IsChecked(ListItemPtr.ToSharedRef(), &FRevertCheckBoxListViewItem::OnIsChecked)
+				SNew(SHorizontalBox)
+				+SHorizontalBox::Slot()
+				.HAlign(HAlign_Left)
+				.AutoWidth()
+				[
+					SAssignNew(CheckBox, SCheckBox)
+					.OnCheckStateChanged(ListItemPtr.ToSharedRef(), &FRevertCheckBoxListViewItem::OnCheckStateChanged)
+					.IsChecked(ListItemPtr.ToSharedRef(), &FRevertCheckBoxListViewItem::OnIsChecked)
+					[
+						SNew(STextBlock)
+						.Text(ListItemPtr->Text)
+					]
+				]
+				+SHorizontalBox::Slot()
+				.HAlign(HAlign_Right)
 				[
-					SNew(STextBlock)				
-					.Text(ListItemPtr->Text)
-				]		
+					SNew(SImage)
+					.Image(FEditorStyle::GetBrush(TEXT("ContentBrowser.ContentDirty")))
+					.Visibility(ListItemPtr.ToSharedRef(), &FRevertCheckBoxListViewItem::OnGetModifiedStateVisibility)
+					.ToolTipText(LOCTEXT("ModifiedFileToolTip","This file has been modified from the source version"))
+				]
 			];
 
 		return Row;
 	}
 
-
 	/** Called when the settings of the dialog are to be accepted*/
 	FReply OKClicked()
 	{
@@ -235,6 +263,11 @@ private:
 
 	bool IsOKEnabled() const
 	{
+		if (bRevertUnchangedFilesOnly)
+		{
+			return true;
+		}
+
 		for (int32 i=0; i<ListViewItemSource.Num(); i++)
 		{
 			if (ListViewItemSource[i]->IsSelected)
@@ -257,65 +290,7 @@ private:
 	/** Called when the user checks or unchecks the revert unchanged checkbox; updates the list view accordingly */
 	void RevertUnchangedToggled( const ESlateCheckBoxState::Type NewCheckedState )
 	{
-		const bool bChecked = (NewCheckedState == ESlateCheckBoxState::Checked);
-		
-		// If this is the first time the user has checked the "Revert Unchanged" checkbox, query source control to find
-		// the packages that are modified from the version on the server. Due to the fact that this is a synchronous and potentially
-		// slow operation, we only do it once upfront, and then cache the results so that future toggling is nearly instant.
-		if ( bChecked && !bAlreadyQueriedModifiedFiles )
-		{
-			TArray<FString> PackagesToCheck;
-			for (int32 i=0; i<ListViewItemSource.Num(); i++)
-			{
-				PackagesToCheck.Add(SourceControlHelpers::PackageFilename(ListViewItemSource[i]->Text));		
-			}
-
-			// Make sure we update the modified state of the files
-			TSharedRef<FUpdateStatus, ESPMode::ThreadSafe> UpdateStatusOperation = ISourceControlOperation::Create<FUpdateStatus>();
-			UpdateStatusOperation->SetUpdateModifiedState(true);
-			ISourceControlModule::Get().GetProvider().Execute(UpdateStatusOperation, PackagesToCheck);
-
-			// Find the files modified from the server version
-			TArray< FSourceControlStateRef > SourceControlStates;
-			ISourceControlModule::Get().GetProvider().GetState( PackagesToCheck, SourceControlStates, EStateCacheUsage::Use );
-
-			ModifiedPackages.Empty();
-			for( TArray< FSourceControlStateRef >::TConstIterator Iter(SourceControlStates); Iter; Iter++ )
-			{
-				FSourceControlStateRef SourceControlState = *Iter;
-				if(SourceControlState->IsModified())
-				{
-					FString PackageName;
-					if ( FPackageName::TryConvertFilenameToLongPackageName(SourceControlState->GetFilename(), PackageName) )
-					{
-						ModifiedPackages.Add(PackageName);
-					}
-				}
-			}
-			
-			bAlreadyQueriedModifiedFiles = true;
-		}
-
-		// Iterate over each list view item, setting its enabled/selected state appropriately based on whether "Revert Unchanged" is
-		// checked or not
-
-		for (int32 i=0; i<ListViewItemSource.Num(); i++)
-		{
-			TSharedPtr<FRevertCheckBoxListViewItem> CurItem = ListViewItemSource[i];
-
-			bool bItemIsModified = false;
-			for ( int32 ModifiedIndex = 0; ModifiedIndex < ModifiedPackages.Num(); ++ModifiedIndex )
-			{
-				if ( (CurItem->Text) == ModifiedPackages[ModifiedIndex] )
-				{
-					bItemIsModified = true;
-					break;
-				}
-			}
-
-			// Disable the item if the checkbox is checked and the item is modified
-			CurItem->IsEnabled = !( bItemIsModified && bChecked );
-		}
+		bRevertUnchangedFilesOnly = (NewCheckedState == ESlateCheckBoxState::Checked);
 	}
 
 	/**
@@ -328,26 +303,56 @@ private:
 		{
 			TSharedPtr<FRevertCheckBoxListViewItem> CurListViewItem = ListViewItemSource[i];
 
-			if ( CurListViewItem->IsEnabled )
+			if (OnGetItemsEnabled())
 			{
 				CurListViewItem->IsSelected = (NewCheckedState == ESlateCheckBoxState::Checked);
 			}
-			else
+		}
+	}
+
+	/** Caches the current state of the files, */
+	void UpdateSCCStatus()
+	{
+		TArray<FString> PackagesToCheck;
+		for ( const auto& CurItem : ListViewItemSource )
+		{
+			PackagesToCheck.Add(SourceControlHelpers::PackageFilename(CurItem->Text));
+		}
+
+		// Make sure we update the modified state of the files
+		TSharedRef<FUpdateStatus, ESPMode::ThreadSafe> UpdateStatusOperation = ISourceControlOperation::Create<FUpdateStatus>();
+		UpdateStatusOperation->SetUpdateModifiedState(true);
+		ISourceControlModule::Get().GetProvider().Execute(UpdateStatusOperation, PackagesToCheck);
+
+		// Find the files modified from the server version
+		TArray< FSourceControlStateRef > SourceControlStates;
+		ISourceControlModule::Get().GetProvider().GetState( PackagesToCheck, SourceControlStates, EStateCacheUsage::Use );
+
+		ModifiedPackages.Empty();
+
+		for( const auto& ControlState : SourceControlStates )
+		{
+			FString PackageName;
+			FPackageName::TryConvertFilenameToLongPackageName(ControlState->GetFilename(), PackageName);
+			for ( const auto& CurItem : ListViewItemSource )
 			{
-				CurListViewItem->IsSelected = false;
+				if (CurItem->Text == PackageName)
+				{
+					CurItem->IsModified = ControlState->IsModified();
+				}
 			}
 		}
 	}
 
-
+	/** Check for whether the list items are enabled or not */
+	bool OnGetItemsEnabled() const
+	{
+		return !bRevertUnchangedFilesOnly;
+	}
 
 	TWeakPtr<SWindow> ParentFrame;
 	ERevertResults::Type DialogResult;
 
-	/** Cache whether the dialog has already queried source control for modified files or not as an optimization */
-	bool bAlreadyQueriedModifiedFiles;
-
-
 	/** ListView for the packages the user can revert */
 	typedef SListView<TSharedPtr<FRevertCheckBoxListViewItem>> SListViewType;
 	TSharedPtr<SListViewType> RevertListView;
@@ -357,6 +362,9 @@ private:
 
 	/** List of package names that are modified from the versions stored in source control; Used as an optimization */
 	TArray<FString> ModifiedPackages;
+
+	/** Flag set by the user to only revert non modified files */
+	bool bRevertUnchangedFilesOnly;
 };
 
 bool FSourceControlWindows::PromptForRevert( const TArray<FString>& InPackageNames )
