#UE4 quick fix to keep package from showing up when loaded from history
TTP#344630

diff --git a/Engine/Source/Developer/AssetTools/Private/AssetTools.cpp b/Engine/Source/Developer/AssetTools/Private/AssetTools.cpp
--- a/Engine/Source/Developer/AssetTools/Private/AssetTools.cpp
+++ b/Engine/Source/Developer/AssetTools/Private/AssetTools.cpp
@@ -926,9 +926,6 @@ void FAssetTools::DiffAgainstDepot( UObject* InObject, const FString& InPackageP
 						if(OldObject != NULL)
 						{
 							/* Set the revision information*/
-							auto Package = InObject->GetOutermost();
-							auto PackageName = Package->GetName();
-
 							FRevisionInfo OldRevision;
 							OldRevision.Changelist = Revision->GetCheckInIdentifier();
 							OldRevision.Date = Revision->GetDate();
diff --git a/Engine/Source/Editor/SourceControlWindows/Private/SSourceControlHistory.cpp b/Engine/Source/Editor/SourceControlWindows/Private/SSourceControlHistory.cpp
--- a/Engine/Source/Editor/SourceControlWindows/Private/SSourceControlHistory.cpp
+++ b/Engine/Source/Editor/SourceControlWindows/Private/SSourceControlHistory.cpp
@@ -169,7 +169,7 @@ static UObject* GetAssetRevisionObject(TSharedPtr<FHistoryTreeItem> HistoryTreeI
 					TGuardValue<bool> DisableCompileOnLoad(GForceDisableBlueprintCompileOnLoad, true);
 
 					// try and load the temporary package
-					AssetPackage = LoadPackage(NULL, *TempPackageName, LOAD_None);
+					AssetPackage = LoadPackage(NULL, *TempPackageName, LOAD_ForDiff);
 				}
 			} // if FileSourceControlState.IsValid()
 		}
diff --git a/Engine/Source/Runtime/Engine/Private/BlueprintGeneratedClass.cpp b/Engine/Source/Runtime/Engine/Private/BlueprintGeneratedClass.cpp
--- a/Engine/Source/Runtime/Engine/Private/BlueprintGeneratedClass.cpp
+++ b/Engine/Source/Runtime/Engine/Private/BlueprintGeneratedClass.cpp
@@ -55,6 +55,14 @@ void UBlueprintGeneratedClass::PostLoad()
 			ClassFlags &= ~CLASS_NotPlaceable;
 		}
 	}
+
+	if (const UPackage* Package = GetOutermost())
+	{
+		if (Package->PackageFlags & PKG_ForDiffing)
+		{
+			ClassFlags |= CLASS_Deprecated;
+		}
+	}
 #endif // WITH_EDITORONLY_DATA
 }
 
