Fix crash reverting source control

diff --git a/Engine/Source/Editor/SourceControlWindows/Private/SSourceControlRevert.cpp b/Engine/Source/Editor/SourceControlWindows/Private/SSourceControlRevert.cpp
--- a/Engine/Source/Editor/SourceControlWindows/Private/SSourceControlRevert.cpp
+++ b/Engine/Source/Editor/SourceControlWindows/Private/SSourceControlRevert.cpp
@@ -407,34 +407,35 @@ bool FSourceControlWindows::PromptForRevert( const TArray<FString>& InPackageNam
 		{
 			TArray<FString> PackagesToRevert;
 			SourceControlWidget->GetPackagesToRevert( PackagesToRevert );
-			check( PackagesToRevert.Num() > 0 );
-
-			// attempt to unload the packages we are about to revert
-			TArray<UPackage*> PackagesToUnload;
-			for ( TArray<FString>::TConstIterator PackageIter( InPackageNames ); PackageIter; ++PackageIter )
+			if (PackagesToRevert.Num() > 0)
 			{
-				UPackage* Package = FindPackage(NULL, **PackageIter);
-				if(Package != NULL)
+				// attempt to unload the packages we are about to revert
+				TArray<UPackage*> PackagesToUnload;
+				for (TArray<FString>::TConstIterator PackageIter(InPackageNames); PackageIter; ++PackageIter)
 				{
-					PackagesToUnload.Add(Package);
+					UPackage* Package = FindPackage(NULL, **PackageIter);
+					if (Package != NULL)
+					{
+						PackagesToUnload.Add(Package);
+					}
 				}
-			}
-			
-			if(PackagesToUnload.Num() == 0 || PackageTools::UnloadPackages(PackagesToUnload))
-			{
-				for( auto PackageIter( PackagesToUnload.CreateIterator() ); PackageIter; ++PackageIter )
+
+				if (PackagesToUnload.Num() == 0 || PackageTools::UnloadPackages(PackagesToUnload))
 				{
-					(*PackageIter)->ClearFlags(RF_WasLoaded);
+					for (auto PackageIter(PackagesToUnload.CreateIterator()); PackageIter; ++PackageIter)
+					{
+						(*PackageIter)->ClearFlags(RF_WasLoaded);
+					}
+
+					SourceControlProvider.Execute(ISourceControlOperation::Create<FRevert>(), SourceControlHelpers::PackageFilenames(PackagesToRevert));
+				}
+				else
+				{
+					FMessageDialog::Open(EAppMsgType::Ok, LOCTEXT("SCC_Revert_PackageUnloadFailed", "Could not unload assets when reverting files"));
 				}
 
-				SourceControlProvider.Execute( ISourceControlOperation::Create<FRevert>(), SourceControlHelpers::PackageFilenames(PackagesToRevert) );
-			}
-			else
-			{
-				FMessageDialog::Open( EAppMsgType::Ok, LOCTEXT("SCC_Revert_PackageUnloadFailed", "Could not unload assets when reverting files") );
+				bReverted = true;
 			}
-			
-			bReverted = true;
 		}
 	}
 
