Fix a memory corruption bug in FSourceControlWindows::PromptForRevert - after calling PackageTools::UnloadPackages the UPackage*'s may be invalid if they have been GC'd, so iterate over their names again and only flag those that are still valid.

diff --git a/Engine/Source/Editor/SourceControlWindows/Private/SSourceControlRevert.cpp b/Engine/Source/Editor/SourceControlWindows/Private/SSourceControlRevert.cpp
--- a/Engine/Source/Editor/SourceControlWindows/Private/SSourceControlRevert.cpp
+++ b/Engine/Source/Editor/SourceControlWindows/Private/SSourceControlRevert.cpp
@@ -422,9 +422,14 @@ bool FSourceControlWindows::PromptForRevert( const TArray<FString>& InPackageNam
 
 				if (PackagesToUnload.Num() == 0 || PackageTools::UnloadPackages(PackagesToUnload))
 				{
-					for (auto PackageIter(PackagesToUnload.CreateIterator()); PackageIter; ++PackageIter)
+					// Iterate over the names again, the UPackage*'s may be invalid if they have been GC'd
+					for (TArray<FString>::TConstIterator PackageIter(InPackageNames); PackageIter; ++PackageIter)
 					{
-						(*PackageIter)->ClearFlags(RF_WasLoaded);
+						UPackage* Package = FindPackage(NULL, **PackageIter);
+						if (Package != NULL)
+						{
+							Package->ClearFlags(RF_WasLoaded);
+						}
 					}
 
 					SourceControlProvider.Execute(ISourceControlOperation::Create<FRevert>(), SourceControlHelpers::PackageFilenames(PackagesToRevert));
