Added a tooltip to show multi-line descriptions in the source control history

TTP# 346882 - EDITOR: Source Control history only shows first line of description

diff --git a/Engine/Source/Editor/SourceControlWindows/Private/SSourceControlHistory.cpp b/Engine/Source/Editor/SourceControlWindows/Private/SSourceControlHistory.cpp
--- a/Engine/Source/Editor/SourceControlWindows/Private/SSourceControlHistory.cpp
+++ b/Engine/Source/Editor/SourceControlWindows/Private/SSourceControlHistory.cpp
@@ -550,17 +550,25 @@ public:
 		}
 		else if (ColumnName == TEXT("Description"))
 		{
-			//cut down to single line
-			FString Description = RevisionListItem->Description;
+			// Cut down the description to a single line for the list view
+			FString SingleLineDescription = RevisionListItem->Description;
 			int32 NewLinePos;
-			if (Description.FindChar(TCHAR('\n'), NewLinePos))
+			if (SingleLineDescription.FindChar(TCHAR('\n'), NewLinePos))
 			{
-				Description = Description.Left(NewLinePos);
+				SingleLineDescription = SingleLineDescription.Left(NewLinePos);
+			}
+
+			// Trim any trailing new-line characters from the description for the tooltip
+			FString TooltipDescription = RevisionListItem->Description;
+			while(TooltipDescription.Len() && FChar::IsLinebreak(TooltipDescription[TooltipDescription.Len() - 1]))
+			{
+				TooltipDescription.RemoveAt(TooltipDescription.Len() - 1, 1, false);
 			}
 
 			return
 				SNew(STextBlock)
-				.Text(Description);
+				.Text(SingleLineDescription)
+				.ToolTipText(TooltipDescription);
 		}
 		else
 		{
