Config files are now presented in the "Check in files" dialog along with package files.

diff --git a/Engine/Source/Editor/MainFrame/Private/Frame/MainFrameActions.cpp b/Engine/Source/Editor/MainFrame/Private/Frame/MainFrameActions.cpp
--- a/Engine/Source/Editor/MainFrame/Private/Frame/MainFrameActions.cpp
+++ b/Engine/Source/Editor/MainFrame/Private/Frame/MainFrameActions.cpp
@@ -192,7 +192,7 @@ void FMainFrameActionCallbacks::ChoosePackagesToSave()
 	FEditorFileUtils::SaveDirtyPackages( bPromptUserToSave, bSaveMapPackages, bSaveContentPackages, bFastSave, bNotifyNoPackagesSaved );
 }
 
-void FMainFrameActionCallbacks::ChoosePackagesToCheckInCompleted(TArray<UPackage*> LoadedPackages, TArray<FString> PackageNames)
+void FMainFrameActionCallbacks::ChoosePackagesToCheckInCompleted(const TArray<UPackage*>& LoadedPackages, const TArray<FString>& PackageNames, const TArray<FString>& ConfigFiles)
 {
 	if (ChoosePackagesToCheckInNotification.IsValid())
 	{
@@ -200,7 +200,7 @@ void FMainFrameActionCallbacks::ChoosePackagesToCheckInCompleted(TArray<UPackage
 	}
 	ChoosePackagesToCheckInNotification.Reset();
 	
-	check(PackageNames.Num() > 0);
+	check(PackageNames.Num() > 0 || ConfigFiles.Num() > 0);
 
 	// Prompt the user to ask if they would like to first save any dirty packages they are trying to check-in
 	const FEditorFileUtils::EPromptReturnCode UserResponse = FEditorFileUtils::PromptForCheckoutAndSave( LoadedPackages, true, true );
@@ -211,7 +211,7 @@ void FMainFrameActionCallbacks::ChoosePackagesToCheckInCompleted(TArray<UPackage
 	if ( bShouldProceed )
 	{
 		FAssetToolsModule& AssetToolsModule = FModuleManager::LoadModuleChecked<FAssetToolsModule>(TEXT("AssetTools"));
-		FSourceControlWindows::PromptForCheckin(PackageNames);
+		FSourceControlWindows::PromptForCheckin(PackageNames, ConfigFiles);
 	}
 	else
 	{
@@ -280,9 +280,18 @@ void FMainFrameActionCallbacks::ChoosePackagesToCheckInCallback(const FSourceCon
 			}
 		}
 
-		if ( PackageNames.Num() > 0 )
+		// Get a list of all the checked out config files
+		TMap<FString, FSourceControlStatePtr> ConfigFileStates;
+		TArray<FString> ConfigFilesToSubmit;
+		FEditorFileUtils::FindAllSubmittableConfigFiles(ConfigFileStates);
+		for (TMap<FString, FSourceControlStatePtr>::TConstIterator It(ConfigFileStates); It; ++It)
 		{
-			ChoosePackagesToCheckInCompleted(LoadedPackages, PackageNames);
+			ConfigFilesToSubmit.Add(It.Key());
+		}
+
+		if ( PackageNames.Num() > 0 || ConfigFilesToSubmit.Num() > 0 )
+		{
+			ChoosePackagesToCheckInCompleted(LoadedPackages, PackageNames, ConfigFilesToSubmit);
 		}
 		else
 		{
@@ -308,10 +317,16 @@ void FMainFrameActionCallbacks::ChoosePackagesToCheckIn()
 			// make sure we update the SCC status of all packages (this could take a long time, so we will run it as a background task)
 			TArray<FString> Packages;
 			FEditorFileUtils::FindAllPackageFiles(Packages);
+
+			// Get list of filenames corresponding to packages
+			TArray<FString> Filenames = SourceControlHelpers::PackageFilenames(Packages);
+
+			// Add game config files to the list
+			FEditorFileUtils::FindAllConfigFiles(Filenames);
 			
 			ISourceControlProvider& SourceControlProvider = ISourceControlModule::Get().GetProvider();
 			FSourceControlOperationRef Operation = ISourceControlOperation::Create<FUpdateStatus>();
-			SourceControlProvider.Execute(Operation, SourceControlHelpers::PackageFilenames(Packages), EConcurrency::Asynchronous, FSourceControlOperationComplete::CreateStatic(&FMainFrameActionCallbacks::ChoosePackagesToCheckInCallback));
+			SourceControlProvider.Execute(Operation, Filenames, EConcurrency::Asynchronous, FSourceControlOperationComplete::CreateStatic(&FMainFrameActionCallbacks::ChoosePackagesToCheckInCallback));
 
 			if (ChoosePackagesToCheckInNotification.IsValid())
 			{
diff --git a/Engine/Source/Editor/MainFrame/Private/Frame/MainFrameActions.h b/Engine/Source/Editor/MainFrame/Private/Frame/MainFrameActions.h
--- a/Engine/Source/Editor/MainFrame/Private/Frame/MainFrameActions.h
+++ b/Engine/Source/Editor/MainFrame/Private/Frame/MainFrameActions.h
@@ -266,7 +266,7 @@ protected:
 	static void ChoosePackagesToCheckInCallback(const FSourceControlOperationRef& InOperation, ECommandResult::Type InResult);
 
 	/** Called when the process has completed and we have packages to check in */
-	static void ChoosePackagesToCheckInCompleted(TArray<UPackage*> LoadedPackages, TArray<FString> PackageNames);
+	static void ChoosePackagesToCheckInCompleted(const TArray<UPackage*>& LoadedPackages, const TArray<FString>& PackageNames, const TArray<FString>& ConfigFiles);
 
 	/** Delegate called when the user has decided to cancel the check in process */
 	static void ChoosePackagesToCheckInCancelled(FSourceControlOperationRef InOperation);
diff --git a/Engine/Source/Editor/SourceControlWindows/Private/SSourceControlSubmit.cpp b/Engine/Source/Editor/SourceControlWindows/Private/SSourceControlSubmit.cpp
--- a/Engine/Source/Editor/SourceControlWindows/Private/SSourceControlSubmit.cpp
+++ b/Engine/Source/Editor/SourceControlWindows/Private/SSourceControlSubmit.cpp
@@ -454,27 +454,27 @@ private:
 	ESlateCheckBoxState::Type	KeepCheckedOut;
 };
 
-static void FindFilesForCheckIn(const TArray<FString>& InPackagesNames, TArray<FString>& OutAddFiles, TArray<FString>& OutOpenFiles)
+static void FindFilesForCheckIn(const TArray<FString>& InFilenames, TArray<FString>& OutAddFiles, TArray<FString>& OutOpenFiles)
 {
 	ISourceControlProvider& SourceControlProvider = ISourceControlModule::Get().GetProvider();
 
 	TArray<FSourceControlStateRef> States;
-	SourceControlProvider.GetState(SourceControlHelpers::PackageFilenames(InPackagesNames), States, EStateCacheUsage::ForceUpdate);
+	SourceControlProvider.GetState(InFilenames, States, EStateCacheUsage::ForceUpdate);
 
-	for( int32 PackageIndex = 0 ; PackageIndex < InPackagesNames.Num() ; ++PackageIndex )
+	for (const FString& Filename : InFilenames)
 	{
-		FSourceControlStatePtr SourceControlState = SourceControlProvider.GetState(SourceControlHelpers::PackageFilename(InPackagesNames[PackageIndex]), EStateCacheUsage::Use);
+		FSourceControlStatePtr SourceControlState = SourceControlProvider.GetState(Filename, EStateCacheUsage::Use);
 		if(SourceControlState.IsValid())
 		{
 			if (SourceControlState->CanCheckIn())
 			{
-				OutOpenFiles.Add(InPackagesNames[PackageIndex]);
+				OutOpenFiles.Add(Filename);
 			}
 			else
 			{
 				if( !SourceControlState->IsSourceControlled() )
 				{
-					OutAddFiles.Add(InPackagesNames[PackageIndex]);
+					OutAddFiles.Add(Filename);
 				}
 			}
 		}
@@ -482,13 +482,17 @@ static void FindFilesForCheckIn(const TArray<FString>& InPackagesNames, TArray<F
 }
 
 
-bool FSourceControlWindows::PromptForCheckin(const TArray<FString>& InPackageNames)
+bool FSourceControlWindows::PromptForCheckin(const TArray<FString>& InPackageNames, const TArray<FString>& InConfigFiles)
 {
 	bool bCheckInSuccess = true;
 
+	// Get filenames for package names
+	TArray<FString> AllFiles = SourceControlHelpers::PackageFilenames(InPackageNames);
+	AllFiles.Append(InConfigFiles);
+
 	TArray<FString> AddFiles;
 	TArray<FString> OpenFiles;
-	FindFilesForCheckIn(InPackageNames, AddFiles, OpenFiles);
+	FindFilesForCheckIn(AllFiles, AddFiles, OpenFiles);
 
 	if (AddFiles.Num() || OpenFiles.Num())
 	{
@@ -520,10 +524,6 @@ bool FSourceControlWindows::PromptForCheckin(const TArray<FString>& InPackageNam
 			//Get description from the dialog
 			SourceControlWidget->FillChangeListDescription(Description, AddFiles, OpenFiles);
 
-			// Convert to source control paths
-			Description.FilesForAdd = SourceControlHelpers::PackageFilenames(Description.FilesForAdd);
-			Description.FilesForSubmit = SourceControlHelpers::PackageFilenames(Description.FilesForSubmit);
-
 			//revert all unchanged files that were submitted
 			if ( Description.FilesForSubmit.Num() > 0 )
 			{
@@ -532,9 +532,7 @@ bool FSourceControlWindows::PromptForCheckin(const TArray<FString>& InPackageNam
 				//make sure all files are still checked out
 				for (int32 VerifyIndex = Description.FilesForSubmit.Num()-1; VerifyIndex >= 0; --VerifyIndex)
 				{
-					FString TempFilename = Description.FilesForSubmit[VerifyIndex];
-					const FString PackageName = FPackageName::FilenameToLongPackageName(TempFilename);
-					FSourceControlStatePtr SourceControlState = SourceControlProvider.GetState(SourceControlHelpers::PackageFilename(PackageName), EStateCacheUsage::Use);
+					FSourceControlStatePtr SourceControlState = SourceControlProvider.GetState(Description.FilesForSubmit[VerifyIndex], EStateCacheUsage::Use);
 					if( SourceControlState.IsValid() && !SourceControlState->IsCheckedOut() && !SourceControlState->IsAdded() )
 					{
 						Description.FilesForSubmit.RemoveAt(VerifyIndex);
diff --git a/Engine/Source/Editor/SourceControlWindows/Public/SourceControlWindows.h b/Engine/Source/Editor/SourceControlWindows/Public/SourceControlWindows.h
--- a/Engine/Source/Editor/SourceControlWindows/Public/SourceControlWindows.h
+++ b/Engine/Source/Editor/SourceControlWindows/Public/SourceControlWindows.h
@@ -11,8 +11,9 @@ public:
 	 * Display check in dialog for the specified packages
 	 *
 	 * @param	InPackageNames	Names of packages to check in
+	 * @param	InConfigFiles	Config filenames to check in
 	 */
-	static bool PromptForCheckin(const TArray<FString>& InPackageNames);
+	static bool PromptForCheckin(const TArray<FString>& InPackageNames, const TArray<FString>& InConfigFiles = TArray<FString>());
 
 	/**
 	 * Display file revision history for the provided packages
diff --git a/Engine/Source/Editor/UnrealEd/Private/FileHelpers.cpp b/Engine/Source/Editor/UnrealEd/Private/FileHelpers.cpp
--- a/Engine/Source/Editor/UnrealEd/Private/FileHelpers.cpp
+++ b/Engine/Source/Editor/UnrealEd/Private/FileHelpers.cpp
@@ -3101,6 +3101,40 @@ void FEditorFileUtils::FindAllSubmittablePackageFiles(TMap<FString, FSourceContr
 	}
 }
 
+void FEditorFileUtils::FindAllConfigFiles(TArray<FString>& OutConfigFiles)
+{
+	TArray<FString> IniFilenames;
+	IFileManager::Get().FindFiles(IniFilenames, *(FPaths::GameConfigDir() / TEXT("*.ini")), true, false);
+	for (const FString& IniFilename : IniFilenames)
+	{
+		OutConfigFiles.Add(FPaths::ConvertRelativePathToFull(FPaths::GameConfigDir() / IniFilename));
+	}
+}
+
+void FEditorFileUtils::FindAllSubmittableConfigFiles(TMap<FString, TSharedPtr<class ISourceControlState, ESPMode::ThreadSafe> >& OutConfigFiles)
+{
+	ISourceControlProvider& SourceControlProvider = ISourceControlModule::Get().GetProvider();
+
+	TArray<FString> ConfigFilenames;
+	FEditorFileUtils::FindAllConfigFiles(ConfigFilenames);
+
+	for (const FString& ConfigFilename : ConfigFilenames)
+	{
+		// Only check files which are intended to be under source control
+		if (FPaths::GetCleanFilename(ConfigFilename) != TEXT("DefaultEditorUserSettings.ini"))
+		{
+			FSourceControlStatePtr SourceControlState = SourceControlProvider.GetState(ConfigFilename, EStateCacheUsage::Use);
+
+			// Only include config files that are currently checked out or packages not under source control
+			if (SourceControlState.IsValid() &&
+				(SourceControlState->IsCheckedOut() || SourceControlState->IsAdded() || (!SourceControlState->IsSourceControlled() && SourceControlState->CanAdd())))
+			{
+				OutConfigFiles.Add(ConfigFilename, SourceControlState);
+			}
+		}
+	}
+}
+
 bool FEditorFileUtils::IsMapPackageAsset(const FString& ObjectPath)
 {
 	FString MapFilePath;
diff --git a/Engine/Source/Editor/UnrealEd/Public/FileHelpers.h b/Engine/Source/Editor/UnrealEd/Public/FileHelpers.h
--- a/Engine/Source/Editor/UnrealEd/Public/FileHelpers.h
+++ b/Engine/Source/Editor/UnrealEd/Public/FileHelpers.h
@@ -343,9 +343,23 @@ public:
 	 * @param	OutPackages		All found package filenames and their source control state
 	 * @param	bIncludeMaps	If true, also adds maps to the list
 	 */
-	UNREALED_API static void FindAllSubmittablePackageFiles(TMap<FString, TSharedPtr<class ISourceControlState, ESPMode::ThreadSafe> >& OutPackages, const bool bIncludeMaps);
+	UNREALED_API static void FindAllSubmittablePackageFiles(TMap<FString, FSourceControlStatePtr>& OutPackages, const bool bIncludeMaps);
 
-	/** 
+	/**
+	 * Looks for config files for the current project.
+	 *
+	 * @param	OutConfigFiles	All found config filenames.
+	 */
+	UNREALED_API static void FindAllConfigFiles(TArray<FString>& OutConfigFiles);
+
+	/**
+	 * Looks for source control submittable config files for the current project.
+	 *
+	 * @param	OutConfigFiles	All found config filenames and their source control state.
+	 */
+	UNREALED_API static void FindAllSubmittableConfigFiles(TMap<FString, FSourceControlStatePtr>& OutConfigFiles);
+
+	/**
 	 * Helper function used to decide whether a package name is a map package or not. Map packages aren't added to the additional package list.
 	 *
 	 * @param	ObjectPath		The path to the package to test
