Made Error message in submit to source control dialog stand out more.

diff --git a/Engine/Source/Editor/SourceControlWindows/Private/SSourceControlSubmit.cpp b/Engine/Source/Editor/SourceControlWindows/Private/SSourceControlSubmit.cpp
--- a/Engine/Source/Editor/SourceControlWindows/Private/SSourceControlSubmit.cpp
+++ b/Engine/Source/Editor/SourceControlWindows/Private/SSourceControlSubmit.cpp
@@ -206,8 +206,8 @@ public:
 					.Visibility(this, &SSourceControlSubmitWidget::IsWarningPanelVisible)
 					.Padding(5)
 					[
-						SNew( STextBlock )
-						.Text( NSLOCTEXT("SourceControl.SubmitPanel", "ChangeListDescWarning", "Changelist description is required to submit") )
+						SNew( SErrorText )
+						.ErrorText( NSLOCTEXT("SourceControl.SubmitPanel", "ChangeListDescWarning", "Changelist description is required to submit") )
 					]
 				]
 				+SVerticalBox::Slot()
