Made Source Control Submit window user-resizable.

diff --git a/Engine/Source/Editor/SourceControlWindows/Private/SSourceControlSubmit.cpp b/Engine/Source/Editor/SourceControlWindows/Private/SSourceControlSubmit.cpp
--- a/Engine/Source/Editor/SourceControlWindows/Private/SSourceControlSubmit.cpp
+++ b/Engine/Source/Editor/SourceControlWindows/Private/SSourceControlSubmit.cpp
@@ -151,9 +151,7 @@ public:
 					]
 				]
 				+SVerticalBox::Slot()
-				.AutoHeight()
 				.Padding(FMargin(5, 0))
-				.MaxHeight(300)
 				[
 					SNew(SBorder)
 					[
@@ -188,7 +186,6 @@ public:
 						.Padding(2)
 						[
 							SNew(SBox)
-							.WidthOverride(500.0f)
 							[
 								SAssignNew( ListView, SListType)
 								.ItemHeight(24)
@@ -498,7 +495,8 @@ bool FSourceControlWindows::PromptForCheckin(const TArray<FString>& InPackageNam
 	{
 		TSharedRef<SWindow> NewWindow = SNew(SWindow)
 			.Title(NSLOCTEXT("SourceControl.SubmitWindow", "Title", "Submit Files"))
-			.SizingRule( ESizingRule::Autosized )
+			.SizingRule(ESizingRule::UserSized)
+			.ClientSize(FVector2D(512, 430))
 			.SupportsMaximize(false)
 			.SupportsMinimize(false);
 
