Fixed resizable source control submit widget on Mac.

diff --git a/Engine/Source/Editor/SourceControlWindows/Private/SSourceControlSubmit.cpp b/Engine/Source/Editor/SourceControlWindows/Private/SSourceControlSubmit.cpp
--- a/Engine/Source/Editor/SourceControlWindows/Private/SSourceControlSubmit.cpp
+++ b/Engine/Source/Editor/SourceControlWindows/Private/SSourceControlSubmit.cpp
@@ -497,7 +497,7 @@ bool FSourceControlWindows::PromptForCheckin(const TArray<FString>& InPackageNam
 			.Title(NSLOCTEXT("SourceControl.SubmitWindow", "Title", "Submit Files"))
 			.SizingRule(ESizingRule::UserSized)
 			.ClientSize(FVector2D(512, 430))
-			.SupportsMaximize(false)
+			.SupportsMaximize(true)
 			.SupportsMinimize(false);
 
 		TSharedRef<SSourceControlSubmitWidget> SourceControlWidget = 
