Fixed code relying on SLATE_TEXT_ATTRIBUTE for STextBlock.

diff --git a/Engine/Source/Editor/SceneOutliner/Private/SSceneOutliner.h b/Engine/Source/Editor/SceneOutliner/Private/SSceneOutliner.h
--- a/Engine/Source/Editor/SceneOutliner/Private/SSceneOutliner.h
+++ b/Engine/Source/Editor/SceneOutliner/Private/SSceneOutliner.h
@@ -298,7 +298,7 @@ namespace SceneOutliner
 		EVisibility GetFilterStatusVisibility() const;
 
 		/** @return	Returns the filter status text */
-		FString GetFilterStatusText() const;
+		FText GetFilterStatusText() const;
 
 		/** @return Returns color for the filter status text message, based on success of search filter */
 		FSlateColor GetFilterStatusTextColor() const;
diff --git a/Engine/Source/Editor/SceneOutliner/Private/SSocketChooser.h b/Engine/Source/Editor/SceneOutliner/Private/SSocketChooser.h
--- a/Engine/Source/Editor/SceneOutliner/Private/SSocketChooser.h
+++ b/Engine/Source/Editor/SceneOutliner/Private/SSocketChooser.h
@@ -91,7 +91,7 @@ public:
 				.FillWidth(1.f)
 				[
 					SNew(STextBlock)
-					.Text(SocketInfo->Description.Name.ToString())
+					.Text(FText::FromName(SocketInfo->Description.Name))
 				]
 			];
 	}
diff --git a/Engine/Source/Editor/SceneOutliner/Private/SceneOutlinerActorInfoColumn.cpp b/Engine/Source/Editor/SceneOutliner/Private/SceneOutlinerActorInfoColumn.cpp
--- a/Engine/Source/Editor/SceneOutliner/Private/SceneOutlinerActorInfoColumn.cpp
+++ b/Engine/Source/Editor/SceneOutliner/Private/SceneOutlinerActorInfoColumn.cpp
@@ -284,10 +284,10 @@ EVisibility FActorInfoColumn::GetColumnDataVisibility( bool bIsClassHyperlink )
 	}
 }
 
-FString FActorInfoColumn::GetTextForItem( TWeakPtr<ITreeItem> TreeItem ) const
+FText FActorInfoColumn::GetTextForItem( TWeakPtr<ITreeItem> TreeItem ) const
 {
 	auto Item = TreeItem.Pin();
-	return Item.IsValid() ? Item->Get(FGetInfo(CurrentMode)) : FString();
+	return Item.IsValid() ? FText::FromString(Item->Get(FGetInfo(CurrentMode))) : FText::GetEmpty();
 }
 
 FText FActorInfoColumn::GetSelectedMode() const
diff --git a/Engine/Source/Editor/SceneOutliner/Private/SceneOutlinerActorInfoColumn.h b/Engine/Source/Editor/SceneOutliner/Private/SceneOutlinerActorInfoColumn.h
--- a/Engine/Source/Editor/SceneOutliner/Private/SceneOutlinerActorInfoColumn.h
+++ b/Engine/Source/Editor/SceneOutliner/Private/SceneOutlinerActorInfoColumn.h
@@ -77,7 +77,7 @@ public:
 	// End ISceneOutlinerColumn Implementation
 	//////////////////////////////////////////////////////////////////////////
 
-	FString GetTextForItem( const TWeakPtr<ITreeItem> TreeItem ) const;
+	FText GetTextForItem( const TWeakPtr<ITreeItem> TreeItem ) const;
 
 private:
 
diff --git a/Engine/Source/Editor/SoundClassEditor/Private/SSoundClassActionMenu.cpp b/Engine/Source/Editor/SoundClassEditor/Private/SSoundClassActionMenu.cpp
--- a/Engine/Source/Editor/SoundClassEditor/Private/SSoundClassActionMenu.cpp
+++ b/Engine/Source/Editor/SoundClassEditor/Private/SSoundClassActionMenu.cpp
@@ -26,7 +26,7 @@ void SSoundClassActionMenuItem::Construct(const FArguments& InArgs, TSharedPtr<F
 		this->ChildSlot
 			[
 				SNew(SHorizontalBox)
-				.ToolTipText( InAction->TooltipDescription )
+				.ToolTipText( FText::FromString(InAction->TooltipDescription) )
 				+ SHorizontalBox::Slot()
 				.AutoWidth()
 				.VAlign(VAlign_Center)
@@ -40,7 +40,7 @@ void SSoundClassActionMenuItem::Construct(const FArguments& InArgs, TSharedPtr<F
 	}
 	else
 	{
-		TSharedRef<SWidget> NewSoundClassWidget = CreateNewSoundClassWidget( InAction->MenuDescription.ToString(), InAction->TooltipDescription, FSlateFontInfo( FPaths::EngineContentDir() / TEXT("Slate/Fonts/Roboto-Regular.ttf"), 9), InAction );
+		TSharedRef<SWidget> NewSoundClassWidget = CreateNewSoundClassWidget( InAction->MenuDescription, FText::FromString(InAction->TooltipDescription), FSlateFontInfo( FPaths::EngineContentDir() / TEXT("Slate/Fonts/Roboto-Regular.ttf"), 9), InAction );
 		// Promote requires 2 'slots'
 		this->ChildSlot
 			[
@@ -49,7 +49,7 @@ void SSoundClassActionMenuItem::Construct(const FArguments& InArgs, TSharedPtr<F
 	}
 }
 
-TSharedRef<SWidget> SSoundClassActionMenuItem::CreateNewSoundClassWidget( const FString& DisplayText, const FString& ToolTip, const FSlateFontInfo& NameFont, TSharedPtr<FEdGraphSchemaAction>& InAction )
+TSharedRef<SWidget> SSoundClassActionMenuItem::CreateNewSoundClassWidget( const FText& DisplayText, const FText& ToolTip, const FSlateFontInfo& NameFont, TSharedPtr<FEdGraphSchemaAction>& InAction )
 {
 	FString ClassName;
 	FSoundClassGraphSchemaAction_NewNode* Action = static_cast<FSoundClassGraphSchemaAction_NewNode*>(InAction.Get());
diff --git a/Engine/Source/Editor/SoundClassEditor/Private/SSoundClassActionMenu.h b/Engine/Source/Editor/SoundClassEditor/Private/SSoundClassActionMenu.h
--- a/Engine/Source/Editor/SoundClassEditor/Private/SSoundClassActionMenu.h
+++ b/Engine/Source/Editor/SoundClassEditor/Private/SSoundClassActionMenu.h
@@ -21,7 +21,7 @@ public:
 
 private:
 	/* Create the new sound class widget */
-	TSharedRef<SWidget> CreateNewSoundClassWidget( const FString& DisplayText, const FString& ToolTip, const FSlateFontInfo& NameFont, TSharedPtr<FEdGraphSchemaAction>& InAction );
+	TSharedRef<SWidget> CreateNewSoundClassWidget( const FText& DisplayText, const FText& ToolTip, const FSlateFontInfo& NameFont, TSharedPtr<FEdGraphSchemaAction>& InAction );
 
 	/** Called when confirming name for new sound class */
 	void OnNewSoundClassNameEntered(const FText& NewText, ETextCommit::Type CommitInfo, TSharedPtr<FEdGraphSchemaAction> InAction );
diff --git a/Engine/Source/Editor/SoundClassEditor/Private/SoundClassEditor.cpp b/Engine/Source/Editor/SoundClassEditor/Private/SoundClassEditor.cpp
--- a/Engine/Source/Editor/SoundClassEditor/Private/SoundClassEditor.cpp
+++ b/Engine/Source/Editor/SoundClassEditor/Private/SoundClassEditor.cpp
@@ -226,7 +226,7 @@ TSharedRef<SGraphEditor> FSoundClassEditor::CreateGraphEditorWidget()
 	}
 
 	FGraphAppearanceInfo AppearanceInfo;
-	AppearanceInfo.CornerText = LOCTEXT("AppearanceCornerText_SoundClass", "SOUND CLASS").ToString();
+	AppearanceInfo.CornerText = LOCTEXT("AppearanceCornerText_SoundClass", "SOUND CLASS");
 
 	SGraphEditor::FGraphEditorEvents InEvents;
 	InEvents.OnSelectionChanged = SGraphEditor::FOnSelectionChanged::CreateSP(this, &FSoundClassEditor::OnSelectedNodesChanged);
diff --git a/Engine/Source/Editor/SoundCueEditor/Private/SoundCueEditor.cpp b/Engine/Source/Editor/SoundCueEditor/Private/SoundCueEditor.cpp
--- a/Engine/Source/Editor/SoundCueEditor/Private/SoundCueEditor.cpp
+++ b/Engine/Source/Editor/SoundCueEditor/Private/SoundCueEditor.cpp
@@ -520,7 +520,7 @@ TSharedRef<SGraphEditor> FSoundCueEditor::CreateGraphEditorWidget()
 	}
 
 	FGraphAppearanceInfo AppearanceInfo;
-	AppearanceInfo.CornerText = LOCTEXT("AppearanceCornerText_SoundCue", "SOUND CUE").ToString();
+	AppearanceInfo.CornerText = LOCTEXT("AppearanceCornerText_SoundCue", "SOUND CUE");
 
 	SGraphEditor::FGraphEditorEvents InEvents;
 	InEvents.OnSelectionChanged = SGraphEditor::FOnSelectionChanged::CreateSP(this, &FSoundCueEditor::OnSelectedNodesChanged);
diff --git a/Engine/Source/Editor/SourceControlWindows/Private/SSourceControlHistory.cpp b/Engine/Source/Editor/SourceControlWindows/Private/SSourceControlHistory.cpp
--- a/Engine/Source/Editor/SourceControlWindows/Private/SSourceControlHistory.cpp
+++ b/Engine/Source/Editor/SourceControlWindows/Private/SSourceControlHistory.cpp
@@ -567,8 +567,8 @@ public:
 
 			return
 				SNew(STextBlock)
-				.Text(SingleLineDescription)
-				.ToolTipText(TooltipDescription);
+				.Text(FText::FromString(SingleLineDescription))
+				.ToolTipText(FText::FromString(TooltipDescription));
 		}
 		else
 		{
@@ -695,28 +695,28 @@ private:
 					.Padding(Padding)
 					[
 						SNew(STextBlock)
-						.Text(NSLOCTEXT("SourceControl.HistoryPanel.Info", "Revision", "Revision:").ToString())
+						.Text(NSLOCTEXT("SourceControl.HistoryPanel.Info", "Revision", "Revision:"))
 					]
 					+SVerticalBox::Slot()
 					.FillHeight(0.25f)
 					.Padding(Padding)
 					[
 						SNew(STextBlock)
-						.Text(NSLOCTEXT("SourceControl.HistoryPanel.Info", "Date", "Date Submitted:").ToString())
+						.Text(NSLOCTEXT("SourceControl.HistoryPanel.Info", "Date", "Date Submitted:"))
 					]
 					+SVerticalBox::Slot()
 					.FillHeight(0.25f)
 					.Padding(Padding)
 					[
 						SNew(STextBlock)
-						.Text(NSLOCTEXT("SourceControl.HistoryPanel.Info", "SubmittedBy", "Submitted By:").ToString())
+						.Text(NSLOCTEXT("SourceControl.HistoryPanel.Info", "SubmittedBy", "Submitted By:"))
 					]
 					+SVerticalBox::Slot()
 					.FillHeight(0.25f)
 					.Padding(Padding)
 					[
 						SNew(STextBlock)
-						.Text(NSLOCTEXT("SourceControl.HistoryPanel.Info", "Action", "Action:").ToString())
+						.Text(NSLOCTEXT("SourceControl.HistoryPanel.Info", "Action", "Action:"))
 					]			
 				]
 				+SHorizontalBox::Slot()
@@ -772,28 +772,28 @@ private:
 					.Padding(Padding)
 					[
 						SNew(STextBlock)
-						.Text(NSLOCTEXT("SourceControl.HistoryPanel.Info", "Changelist", "Changelist:").ToString())
+						.Text(NSLOCTEXT("SourceControl.HistoryPanel.Info", "Changelist", "Changelist:"))
 					]
 					+SVerticalBox::Slot()
 					.FillHeight(0.25f)
 					.Padding(Padding)
 					[
 						SNew(STextBlock)
-						.Text(NSLOCTEXT("SourceControl.HistoryPanel.Info", "Workspace", "Workspace:").ToString())
+						.Text(NSLOCTEXT("SourceControl.HistoryPanel.Info", "Workspace", "Workspace:"))
 					]
 					+SVerticalBox::Slot()
 					.FillHeight(0.25f)
 					.Padding(Padding)
 					[
 						SNew(STextBlock)
-						.Text(NSLOCTEXT("SourceControl.HistoryPanel.Info", "FileSize", "File Size:").ToString())
+						.Text(NSLOCTEXT("SourceControl.HistoryPanel.Info", "FileSize", "File Size:"))
 					]
 					+SVerticalBox::Slot()
 					.FillHeight(0.25f)
 					.Padding(Padding)
 					[
 						SNew(STextBlock)
-						.Text(NSLOCTEXT("SourceControl.HistoryPanel.Info", "BranchedFrom", "Branched From:").ToString())
+						.Text(NSLOCTEXT("SourceControl.HistoryPanel.Info", "BranchedFrom", "Branched From:"))
 					]
 				]
 				+SHorizontalBox::Slot()
@@ -844,7 +844,7 @@ private:
 			.Padding(Padding, 10, Padding, 5)
 			[
 				SNew(STextBlock)
-				.Text(NSLOCTEXT("SourceControl.HistoryPanel.Info", "Description", "Description:").ToString())
+				.Text(NSLOCTEXT("SourceControl.HistoryPanel.Info", "Description", "Description:"))
 			]
 			+SVerticalBox::Slot()
 			.FillHeight(1.0f)
@@ -866,93 +866,103 @@ private:
 	END_SLATE_FUNCTION_BUILD_OPTIMIZATION
 
 	/** Get the last selected revision's revision number */
-	FString GetRevisionNumber() const
+	FText GetRevisionNumber() const
 	{
 		if(LastSelectedRevisionItem.IsValid())
 		{
-			return FString::Printf(TEXT("%d"), LastSelectedRevisionItem.Pin()->RevisionNumber);
+			static const FNumberFormattingOptions RevisionFormatOptions = FNumberFormattingOptions()
+				.SetUseGrouping(false); // don't group the revision# as Perforce doesn't display it that way
+			return FText::AsNumber(LastSelectedRevisionItem.Pin()->RevisionNumber, &RevisionFormatOptions);
 		}
-		return FString("");
+		return FText::GetEmpty();
 	}
 
 	/** Get the last selected revision's date */
-	FString GetDate() const
+	FText GetDate() const
 	{
 		if(LastSelectedRevisionItem.IsValid() && LastSelectedRevisionItem.Pin()->Date != 0)
 		{
-			return FText::AsDateTime( LastSelectedRevisionItem.Pin()->Date ).ToString();
+			return FText::AsDateTime(LastSelectedRevisionItem.Pin()->Date);
 		}
-		return FString("");
+		return FText::GetEmpty();
 	}
 
 	/** Get the last selected revision's username */
-	FString GetUserName() const
+	FText GetUserName() const
 	{
 		if(LastSelectedRevisionItem.IsValid())
 		{
-			return LastSelectedRevisionItem.Pin()->UserName;
+			return FText::FromString(LastSelectedRevisionItem.Pin()->UserName);
 		}
-		return FString("");
+		return FText::GetEmpty();
 	}
 
 	/** Get the last selected revision's revision number */
-	FString GetAction() const
+	FText GetAction() const
 	{
 		if(LastSelectedRevisionItem.IsValid())
 		{
-			return LastSelectedRevisionItem.Pin()->Action;
+			return FText::FromString(LastSelectedRevisionItem.Pin()->Action);
 		}
-		return FString("");
+		return FText::GetEmpty();
 	}
 
 	/** Get the last selected revision's changelist number */
-	FString GetChangelistNumber() const
+	FText GetChangelistNumber() const
 	{
 		if(LastSelectedRevisionItem.IsValid())
 		{
-			return FString::Printf(TEXT("%d"), LastSelectedRevisionItem.Pin()->ChangelistNumber);
+			static const FNumberFormattingOptions ChangelistFormatOptions = FNumberFormattingOptions()
+				.SetUseGrouping(false); // don't group the CL# as Perforce doesn't display it that way
+			return FText::AsNumber(LastSelectedRevisionItem.Pin()->ChangelistNumber, &ChangelistFormatOptions);
 		}
-		return FString("");
+		return FText::GetEmpty();
 	}
 
 	/** Get the last selected revision's client spec */
-	FString GetClientSpec() const
+	FText GetClientSpec() const
 	{
 		if(LastSelectedRevisionItem.IsValid())
 		{
-			return LastSelectedRevisionItem.Pin()->ClientSpec;
+			return FText::FromString(LastSelectedRevisionItem.Pin()->ClientSpec);
 		}
-		return FString("");
+		return FText::GetEmpty();
 	}
 
 	/** Get the last selected revision's file size */
-	FString GetFileSize() const
+	FText GetFileSize() const
 	{
 		if(LastSelectedRevisionItem.IsValid())
 		{
-			return FString::Printf(TEXT("%.1f MB"),((float)LastSelectedRevisionItem.Pin()->FileSize) / (1024.f * 1024.f));
+			static const FNumberFormattingOptions FileSizeFormatOptions = FNumberFormattingOptions()
+				.SetMinimumFractionalDigits(1)
+				.SetMaximumFractionalDigits(1);
+			return FText::Format(
+				NSLOCTEXT("SourceControlHistory", "FileSizeInMBFmt", "{0} MB"), 
+				FText::AsNumber(((float)LastSelectedRevisionItem.Pin()->FileSize) / (1024.f * 1024.f), &FileSizeFormatOptions)
+				);
 		}
-		return FString("");
+		return FText::GetEmpty();
 	}
 
 	/** Get the last selected revision's description */
-	FString GetDescription() const
+	FText GetDescription() const
 	{
 		if(LastSelectedRevisionItem.IsValid())
 		{
-			return LastSelectedRevisionItem.Pin()->Description;
+			return FText::FromString(LastSelectedRevisionItem.Pin()->Description);
 		}
-		return FString("");
+		return FText::GetEmpty();
 	}
 
 	/** Get the last selected revision's description */
-	FString GetBranchedFrom() const
+	FText GetBranchedFrom() const
 	{
 		if(LastSelectedRevisionItem.IsValid())
 		{
-			return LastSelectedRevisionItem.Pin()->BranchSource;
+			return FText::FromString(LastSelectedRevisionItem.Pin()->BranchSource);
 		}
-		return FString("");
+		return FText::GetEmpty();
 	}
 
 	/**
@@ -975,7 +985,7 @@ private:
 					[
 						SNew( STextBlock )
 						.Font( FEditorStyle::GetFontStyle( TEXT("BoldFont") ))
-						.Text( FileListItem->FileName )
+						.Text( FText::FromString(FileListItem->FileName) )
 					]
 				]
 				.OnDragDetected(this, &SSourceControlHistoryWidget::OnRowDragDetected)
diff --git a/Engine/Source/Editor/SourceControlWindows/Private/SSourceControlRevert.cpp b/Engine/Source/Editor/SourceControlWindows/Private/SSourceControlRevert.cpp
--- a/Engine/Source/Editor/SourceControlWindows/Private/SSourceControlRevert.cpp
+++ b/Engine/Source/Editor/SourceControlWindows/Private/SSourceControlRevert.cpp
@@ -103,7 +103,7 @@ public:
 				.Padding(10)
 				[
 					SNew(STextBlock)
-					.Text(NSLOCTEXT("SourceControl.Revert", "SelectFiles", "Select the files that should be reverted below").ToString())
+					.Text(NSLOCTEXT("SourceControl.Revert", "SelectFiles", "Select the files that should be reverted below"))
 				]
 				+SVerticalBox::Slot()
 				.AutoHeight()
@@ -118,7 +118,7 @@ public:
 						.IsEnabled(this, &SSourceControlRevertWidget::OnGetItemsEnabled)
 						[
 							SNew(STextBlock)
-							.Text(NSLOCTEXT("SourceControl.Revert", "ListHeader", "File Name").ToString())
+							.Text(NSLOCTEXT("SourceControl.Revert", "ListHeader", "File Name"))
 						]
 					]
 				]
@@ -152,7 +152,7 @@ public:
 						.OnCheckStateChanged(this, &SSourceControlRevertWidget::RevertUnchangedToggled)
 						[
 							SNew(STextBlock)
-							.Text(NSLOCTEXT("SourceControl.Revert", "RevertUnchanged", "Revert Unchanged Only").ToString())
+							.Text(NSLOCTEXT("SourceControl.Revert", "RevertUnchanged", "Revert Unchanged Only"))
 						]
 					]
 					+SHorizontalBox::Slot()
@@ -236,7 +236,7 @@ private:
 					.IsChecked(ListItemPtr.ToSharedRef(), &FRevertCheckBoxListViewItem::OnIsChecked)
 					[
 						SNew(STextBlock)
-						.Text(ListItemPtr->Text)
+						.Text(FText::FromString(ListItemPtr->Text))
 					]
 				]
 				+SHorizontalBox::Slot()
diff --git a/Engine/Source/Editor/StatsViewer/Private/SStatsViewer.cpp b/Engine/Source/Editor/StatsViewer/Private/SStatsViewer.cpp
--- a/Engine/Source/Editor/StatsViewer/Private/SStatsViewer.cpp
+++ b/Engine/Source/Editor/StatsViewer/Private/SStatsViewer.cpp
@@ -749,7 +749,7 @@ bool SStatsViewer::AreStatsDisplayed( TSharedRef<IStatsPage> StatsPage ) const
 	return CurrentStats == StatsPage;
 }
 
-FString SStatsViewer::OnGetFilterComboButtonLabel() const
+FText SStatsViewer::OnGetFilterComboButtonLabel() const
 {
 	if( CurrentStats.IsValid() )
 	{
@@ -768,14 +768,14 @@ FString SStatsViewer::OnGetFilterComboButtonLabel() const
 
 					FFormatNamedArguments Arguments;
 					Arguments.Add(TEXT("FilterName"), FText::FromString(FilterName));
-					return FText::Format( LOCTEXT( "FilterSelected", "Filter: {FilterName}" ), Arguments ).ToString();
+					return FText::Format( LOCTEXT( "FilterSelected", "Filter: {FilterName}" ), Arguments );
 				}
 				++ColumnIndex;
 			}
 		}
 	}
 
-	return LOCTEXT( "Filter", "Filter" ).ToString();
+	return LOCTEXT( "Filter", "Filter" );
 }
 
 void SStatsViewer::SetObjectSet( int32 InSetIndex )
diff --git a/Engine/Source/Editor/StatsViewer/Private/SStatsViewer.h b/Engine/Source/Editor/StatsViewer/Private/SStatsViewer.h
--- a/Engine/Source/Editor/StatsViewer/Private/SStatsViewer.h
+++ b/Engine/Source/Editor/StatsViewer/Private/SStatsViewer.h
@@ -87,7 +87,7 @@ private:
 	bool AreStatsDisplayed( TSharedRef< class IStatsPage > StatsPage ) const;
 
 	/** Get the label for the filter combo button */
-	FString OnGetFilterComboButtonLabel() const;
+	FText OnGetFilterComboButtonLabel() const;
 
 	/** 
 	 * Set the currently displayed object set 
diff --git a/Engine/Source/Editor/StatsViewer/Private/StatsCustomColumn.cpp b/Engine/Source/Editor/StatsViewer/Private/StatsCustomColumn.cpp
--- a/Engine/Source/Editor/StatsViewer/Private/StatsCustomColumn.cpp
+++ b/Engine/Source/Editor/StatsViewer/Private/StatsCustomColumn.cpp
@@ -60,7 +60,7 @@ TSharedPtr< SWidget > FStatsCustomColumn::CreateColumnLabel( const TSharedRef< I
 			[
 				SNew( STextBlock )
 				.Font( FEditorStyle::GetFontStyle( Style ) )
-				.Text( Column->GetDisplayName().ToString() )
+				.Text( Column->GetDisplayName() )
 			]
 			+SVerticalBox::Slot()
 			.AutoHeight()
@@ -79,7 +79,7 @@ TSharedPtr< SWidget > FStatsCustomColumn::CreateColumnLabel( const TSharedRef< I
 			[
 				SNew( STextBlock )
 				.Font( FEditorStyle::GetFontStyle( Style ) )
-				.Text( Column->GetDisplayName().ToString() )
+				.Text( Column->GetDisplayName() )
 			];
 	}
 }
diff --git a/Engine/Source/Editor/TaskBrowser/Private/STaskBrowser.cpp b/Engine/Source/Editor/TaskBrowser/Private/STaskBrowser.cpp
--- a/Engine/Source/Editor/TaskBrowser/Private/STaskBrowser.cpp
+++ b/Engine/Source/Editor/TaskBrowser/Private/STaskBrowser.cpp
@@ -486,7 +486,7 @@ public:
 			.Padding(2)
 			[
 				SNew( STextBlock )
-				.Text( Entry )
+				.Text( FText::FromString(Entry) )
 			];
 	}
 
diff --git a/Engine/Source/Editor/TaskBrowser/Private/STaskColumn.cpp b/Engine/Source/Editor/TaskBrowser/Private/STaskColumn.cpp
--- a/Engine/Source/Editor/TaskBrowser/Private/STaskColumn.cpp
+++ b/Engine/Source/Editor/TaskBrowser/Private/STaskColumn.cpp
@@ -31,7 +31,7 @@ void STaskColumn::Construct(const FArguments& InArgs)
 	];
 }
 
-FString STaskColumn::GetFieldNameLoc( const EField InField )
+FText STaskColumn::GetFieldNameLoc( const EField InField )
 {
 	// Get a localized column header button text
 	check( InField > EField::Invalid );
@@ -39,19 +39,19 @@ FString STaskColumn::GetFieldNameLoc( const EField InField )
 	switch( InField )
 	{
 	case EField::Number:
-		return LOCTEXT("Number", "Number").ToString();
+		return LOCTEXT("Number", "Number");
 	case EField::Priority:
-		return LOCTEXT("Priority", "Priority").ToString();
+		return LOCTEXT("Priority", "Priority");
 	case EField::Summary:
-		return LOCTEXT("Summary", "Summary").ToString();
+		return LOCTEXT("Summary", "Summary");
 	case EField::Status:
-		return LOCTEXT("Status", "Status").ToString();
+		return LOCTEXT("Status", "Status");
 	case EField::CreatedBy:
-		return LOCTEXT("CreatedBy", "Created By").ToString();
+		return LOCTEXT("CreatedBy", "Created By");
 	case EField::AssignedTo:
-		return LOCTEXT("AssignedTo", "Assigned To").ToString();
+		return LOCTEXT("AssignedTo", "Assigned To");
 	}
-	return LOCTEXT("BadField", "Bad Field").ToString();
+	return LOCTEXT("BadField", "Bad Field");
 }
 
 FReply STaskColumn::OnTaskColumnClicked()
diff --git a/Engine/Source/Editor/TaskBrowser/Private/STaskColumn.h b/Engine/Source/Editor/TaskBrowser/Private/STaskColumn.h
--- a/Engine/Source/Editor/TaskBrowser/Private/STaskColumn.h
+++ b/Engine/Source/Editor/TaskBrowser/Private/STaskColumn.h
@@ -26,7 +26,7 @@ public:
 private:
 
 	/** Helper func to get the field from the localized name */
-	FString GetFieldNameLoc( const EField InField );
+	FText GetFieldNameLoc( const EField InField );
 
 	/** Handle when the column in the table has been clicked */
 	FReply OnTaskColumnClicked();
diff --git a/Engine/Source/Editor/TextureEditor/Private/Widgets/STextureEditorViewport.cpp b/Engine/Source/Editor/TextureEditor/Private/Widgets/STextureEditorViewport.cpp
--- a/Engine/Source/Editor/TextureEditor/Private/Widgets/STextureEditorViewport.cpp
+++ b/Engine/Source/Editor/TextureEditor/Private/Widgets/STextureEditorViewport.cpp
@@ -56,11 +56,11 @@ void STextureEditorViewport::Construct( const FArguments& InArgs, const TSharedR
 		ZoomMenuBuilder.AddMenuEntry(LOCTEXT("ZoomFitAction", "Scale To Fit"), LOCTEXT("ZoomFillActionHint", "Scale the texture to fit the viewport."), FSlateIcon(), ZoomFitAction, NAME_None, EUserInterfaceActionType::ToggleButton);
 	}
 
-	FString TextureName;
+	FText TextureName = FText::GetEmpty();
 	
 	if (InToolkit->GetTexture() != nullptr)
 	{
-		TextureName = InToolkit->GetTexture()->GetFName().ToString();
+		TextureName = FText::FromName(InToolkit->GetTexture()->GetFName());
 	}
 
 	this->ChildSlot
diff --git a/Engine/Source/Editor/UMGEditor/Private/WidgetBlueprintEditor.cpp b/Engine/Source/Editor/UMGEditor/Private/WidgetBlueprintEditor.cpp
--- a/Engine/Source/Editor/UMGEditor/Private/WidgetBlueprintEditor.cpp
+++ b/Engine/Source/Editor/UMGEditor/Private/WidgetBlueprintEditor.cpp
@@ -675,7 +675,7 @@ FGraphAppearanceInfo FWidgetBlueprintEditor::GetGraphAppearance() const
 
 	if ( GetBlueprintObj()->IsA(UWidgetBlueprint::StaticClass()) )
 	{
-		AppearanceInfo.CornerText = LOCTEXT("AppearanceCornerText", "WIDGET BLUEPRINT").ToString();
+		AppearanceInfo.CornerText = LOCTEXT("AppearanceCornerText", "WIDGET BLUEPRINT");
 	}
 
 	return AppearanceInfo;
diff --git a/Engine/Source/Editor/UndoHistory/Private/Widgets/SUndoHistoryTableRow.h b/Engine/Source/Editor/UndoHistory/Private/Widgets/SUndoHistoryTableRow.h
--- a/Engine/Source/Editor/UndoHistory/Private/Widgets/SUndoHistoryTableRow.h
+++ b/Engine/Source/Editor/UndoHistory/Private/Widgets/SUndoHistoryTableRow.h
@@ -58,7 +58,7 @@ public:
 				[
 					SNew(STextBlock)
 						.ColorAndOpacity(this, &SUndoHistoryTableRow::HandleTitleTextColorAndOpacity)
-						.Text(Title.ToString())
+						.Text(Title)
 				];
 		}
 
diff --git a/Engine/Source/Editor/UnrealEd/Private/AssetThumbnail.cpp b/Engine/Source/Editor/UnrealEd/Private/AssetThumbnail.cpp
--- a/Engine/Source/Editor/UnrealEd/Private/AssetThumbnail.cpp
+++ b/Engine/Source/Editor/UnrealEd/Private/AssetThumbnail.cpp
@@ -480,7 +480,7 @@ private:
 		return false;
 	}
 
-	FString GetLabelText() const
+	FText GetLabelText() const
 	{
 		if ( Label == EThumbnailLabel::ClassName )
 		{
@@ -491,10 +491,10 @@ private:
 			return GetAssetDisplayName();
 		}
 
-		return FString();
+		return FText::GetEmpty();
 	}
 
-	FString GetDisplayNameForClass( UClass* Class ) const
+	FText GetDisplayNameForClass( UClass* Class ) const
 	{
 		FText ClassDisplayName;
 		if ( Class )
@@ -512,10 +512,10 @@ private:
 			}
 		}
 
-		return ClassDisplayName.ToString();
+		return ClassDisplayName;
 	}
 
-	FString GetAssetClassDisplayName() const
+	FText GetAssetClassDisplayName() const
 	{
 		const FAssetData& AssetData = AssetThumbnail->GetAssetData();
 		FString AssetClass = AssetData.AssetClass.ToString();
@@ -526,10 +526,10 @@ private:
 			return GetDisplayNameForClass( Class );
 		}
 
-		return AssetClass;
+		return FText::FromString(AssetClass);
 	}
 
-	FString GetAssetDisplayName() const
+	FText GetAssetDisplayName() const
 	{
 		const FAssetData& AssetData = AssetThumbnail->GetAssetData();
 
@@ -539,7 +539,7 @@ private:
 			return GetDisplayNameForClass( Class );
 		}
 
-		return AssetData.AssetName.ToString();
+		return FText::FromName(AssetData.AssetName);
 	}
 
 	void OnRenderedThumbnailChanged( bool bInHasRenderedThumbnail )
diff --git a/Engine/Source/Editor/UnrealEd/Private/ConsolidateWindow.cpp b/Engine/Source/Editor/UnrealEd/Private/ConsolidateWindow.cpp
--- a/Engine/Source/Editor/UnrealEd/Private/ConsolidateWindow.cpp
+++ b/Engine/Source/Editor/UnrealEd/Private/ConsolidateWindow.cpp
@@ -442,7 +442,7 @@ TSharedRef<ITableRow> SConsolidateToolWidget::OnGenerateRowForList( TSharedPtr<F
 			.OnCheckStateChanged( ListItemPtr.ToSharedRef(), &FListItem::OnAssetSelected )
 			[
 				SNew (STextBlock)
-				.Text( ListItemPtr->GetObjectName() )
+				.Text( FText::FromString(ListItemPtr->GetObjectName()) )
 			]
 		];
 }
diff --git a/Engine/Source/Editor/UnrealEd/Private/Dialogs/Dialogs.cpp b/Engine/Source/Editor/UnrealEd/Private/Dialogs/Dialogs.cpp
--- a/Engine/Source/Editor/UnrealEd/Private/Dialogs/Dialogs.cpp
+++ b/Engine/Source/Editor/UnrealEd/Private/Dialogs/Dialogs.cpp
@@ -24,7 +24,7 @@ class SChoiceDialog : public SCompoundWidget
 public:
 	SLATE_BEGIN_ARGS( SChoiceDialog )	{}
 		SLATE_ATTRIBUTE(TSharedPtr<SWindow>, ParentWindow)
-		SLATE_ATTRIBUTE(FString, Message)	
+		SLATE_ATTRIBUTE(FText, Message)	
 		SLATE_ATTRIBUTE(float, WrapMessageAt)
 		SLATE_ATTRIBUTE(EAppMsgType::Type, MessageType)
 	SLATE_END_ARGS()
@@ -224,7 +224,7 @@ protected:
 	 */
 	void CopyMessageToClipboard( )
 	{
-		FPlatformMisc::ClipboardCopy( *MyMessage.Get() );
+		FPlatformMisc::ClipboardCopy( *MyMessage.Get().ToString() );
 	}
 
 
@@ -257,7 +257,7 @@ private:
 
 	EAppReturnType::Type Response;
 	TSharedPtr<SWindow> ParentWindow;
-	TAttribute<FString> MyMessage;
+	TAttribute<FText> MyMessage;
 };
 
 
@@ -272,7 +272,7 @@ void CreateMsgDlgWindow(TSharedPtr<SWindow>& OutWindow, TSharedPtr<SChoiceDialog
 
 	OutDialog = SNew(SChoiceDialog)
 		.ParentWindow(OutWindow)
-		.Message(InMessage.ToString())
+		.Message(InMessage)
 		.WrapMessageAt(512.0f)
 		.MessageType(InMessageType);
 
@@ -312,7 +312,7 @@ class SModalDialog : public SCompoundWidget
 {
 public:
 	SLATE_BEGIN_ARGS( SModalDialog ){}
-		SLATE_ARGUMENT(FString,Message)
+		SLATE_ARGUMENT(FText,Message)
 	SLATE_END_ARGS()
 
 	void Construct( const FArguments& InArgs )
@@ -374,7 +374,7 @@ public:
 	{
 		if (InKeyEvent.GetKey() == EKeys::C && InKeyEvent.IsControlDown())
 		{
-			FPlatformMisc::ClipboardCopy( *MyMessage.Get() );
+			FPlatformMisc::ClipboardCopy( *MyMessage.Get().ToString() );
 			return FReply::Handled();
 		}
 		return FReply::Unhandled();
@@ -398,7 +398,7 @@ private:
 
 	TSharedPtr<SWindow> MyWindow;
 	bool bUserResponse;
-	TAttribute<FString> MyMessage;
+	TAttribute<FText> MyMessage;
 };
 
 /**
diff --git a/Engine/Source/Editor/UnrealEd/Private/Dialogs/DlgSoundWaveOptions.cpp b/Engine/Source/Editor/UnrealEd/Private/Dialogs/DlgSoundWaveOptions.cpp
--- a/Engine/Source/Editor/UnrealEd/Private/Dialogs/DlgSoundWaveOptions.cpp
+++ b/Engine/Source/Editor/UnrealEd/Private/Dialogs/DlgSoundWaveOptions.cpp
@@ -117,7 +117,7 @@ public:
 		{
 			TableRowContent =
 				SNew( STextBlock )
-					.Text(FString::Printf( TEXT( "%d" ), Item->PreviewInfo->QualitySetting ));
+					.Text(FText::AsNumber(Item->PreviewInfo->QualitySetting));
 		}
 		else if( ColumnName == SoundWaveCompressionOptions::ColumnID_OriginalDataSizeLabel )
 		{
@@ -151,24 +151,34 @@ public:
 	}
 
 private:
+	/** Common code to format a data size percentage string */
+	static FText GetFormattedDataSizeText(const int32 InSize, const int32 InOriginalSize)
+	{
+		static const FNumberFormattingOptions DataSizeFormat = FNumberFormattingOptions()
+			.SetMinimumFractionalDigits(2)
+			.SetMaximumFractionalDigits(2);
+		static const FNumberFormattingOptions OriginalDataSizePercFormat = FNumberFormattingOptions()
+			.SetMinimumFractionalDigits(1)
+			.SetMaximumFractionalDigits(1);
+
+		const float SizePerc = (InOriginalSize == 0) ? 0.0f : static_cast<float>(InSize) / static_cast<float>(InOriginalSize);
+		return FText::Format(
+			NSLOCTEXT("SoundWaveCompressionListRow", "DataSizePercentageOriginalFmt", "{0} ({1})"), 
+			FText::AsNumber(static_cast<float>(InSize) / 1024.0f, &DataSizeFormat), 
+			FText::AsPercent(SizePerc, &OriginalDataSizePercFormat)
+			);
+	}
+
 	/** Retrieves the original data size string. */
-	FString GetOriginalDataSize() const
+	FText GetOriginalDataSize() const
 	{
-		if(Item->PreviewInfo->OriginalSize == 0)
-		{
-			return FString::Printf( TEXT( "%.2f (0.0%%)" ), Item->PreviewInfo->OriginalSize / 1024.0f );
-		}
-		return FString::Printf( TEXT( "%.2f (%.1f%%)" ), Item->PreviewInfo->OriginalSize / 1024.0f, Item->PreviewInfo->OriginalSize * 100.0f / Item->PreviewInfo->OriginalSize);
+		return GetFormattedDataSizeText(Item->PreviewInfo->OriginalSize, Item->PreviewInfo->OriginalSize);
 	}
 
 	/** Retrieves the Vorbis data size string. */
-	FString GetVorbisDataSize() const
+	FText GetVorbisDataSize() const
 	{
-		if(Item->PreviewInfo->OriginalSize == 0)
-		{
-			return FString::Printf( TEXT( "%.2f (0.0%%)" ), Item->PreviewInfo->OggVorbisSize / 1024.0f );
-		}
-		return FString::Printf( TEXT( "%.2f (%.1f%%)" ), Item->PreviewInfo->OggVorbisSize / 1024.0f, Item->PreviewInfo->OggVorbisSize * 100.0f / Item->PreviewInfo->OriginalSize);
+		return GetFormattedDataSizeText(Item->PreviewInfo->OggVorbisSize, Item->PreviewInfo->OriginalSize);
 	}
 
 	/** Plays the Vorbis sound wave. */
@@ -182,13 +192,9 @@ private:
 	}
 
 	/** Retrieves the XMA data size string. */
-	FString GetXMADataSize() const
+	FText GetXMADataSize() const
 	{
-		if(Item->PreviewInfo->OriginalSize == 0)
-		{
-			return FString::Printf( TEXT( "%.2f (0.0%%)" ), Item->PreviewInfo->XMASize / 1024.0f );
-		}
-		return FString::Printf( TEXT( "%.2f (%.1f%%)" ), Item->PreviewInfo->XMASize / 1024.0f, Item->PreviewInfo->XMASize * 100.0f / Item->PreviewInfo->OriginalSize);
+		return GetFormattedDataSizeText(Item->PreviewInfo->XMASize, Item->PreviewInfo->OriginalSize);
 	}
 
 	/** Plays the XMA sound wave. */
@@ -202,13 +208,9 @@ private:
 	}
 
 	/** Retrieves the PS3 data size string. */
-	FString GetPS3DataSize() const
+	FText GetPS3DataSize() const
 	{
-		if(Item->PreviewInfo->OriginalSize == 0)
-		{
-			return FString::Printf( TEXT( "%.2f (0.0%%)" ), Item->PreviewInfo->PS3Size / 1024.0f );
-		}
-		return FString::Printf( TEXT( "%.2f (%.1f%%)" ), Item->PreviewInfo->PS3Size / 1024.0f, Item->PreviewInfo->PS3Size * 100.0f / Item->PreviewInfo->OriginalSize);
+		return GetFormattedDataSizeText(Item->PreviewInfo->PS3Size, Item->PreviewInfo->OriginalSize);
 	}
 
 	/** Plays the PS3 sound wave. */
@@ -402,9 +404,13 @@ TOptional<float> SSoundWaveCompressionOptions::GetProgress() const
 	return 0.0f;
 }
 
-FString SSoundWaveCompressionOptions::GetProgressString() const
+FText SSoundWaveCompressionOptions::GetProgressString() const
 {
-	return FString::Printf(TEXT("%d/%d"), SoundPreviewThreadRunnable->GetIndex(), SoundPreviewThreadRunnable->GetCount());
+	return FText::Format(
+		NSLOCTEXT("SoundWaveOptions", "CompressionProgressXOfYFmt", "{0}/{1}"), 
+		FText::AsNumber(SoundPreviewThreadRunnable->GetIndex()), 
+		FText::AsNumber(SoundPreviewThreadRunnable->GetCount())
+		);
 }
 
 EVisibility SSoundWaveCompressionOptions::IsProgressBarVisible() const
diff --git a/Engine/Source/Editor/UnrealEd/Private/Dialogs/SDeleteAssetsDialog.cpp b/Engine/Source/Editor/UnrealEd/Private/Dialogs/SDeleteAssetsDialog.cpp
--- a/Engine/Source/Editor/UnrealEd/Private/Dialogs/SDeleteAssetsDialog.cpp
+++ b/Engine/Source/Editor/UnrealEd/Private/Dialogs/SDeleteAssetsDialog.cpp
@@ -51,13 +51,13 @@ public:
 				.Padding( 3, 0, 0, 0 )
 				[
 					SNew( STextBlock )
-					.Text(Item->GetObject()->GetName())
+					.Text(FText::FromString(Item->GetObject()->GetName()))
 				];
 		}
 		else if ( ColumnName == DeleteAssetsView::ColumnID_AssetClass )
 		{
 			return SNew( STextBlock )
-				.Text(Item->GetObject()->GetClass()->GetName());
+				.Text(FText::FromString(Item->GetObject()->GetClass()->GetName()));
 		}
 		else if ( ColumnName == DeleteAssetsView::ColumnID_DiskReferences )
 		{
diff --git a/Engine/Source/Editor/UnrealEd/Private/EdGraphNode_Comment.cpp b/Engine/Source/Editor/UnrealEd/Private/EdGraphNode_Comment.cpp
--- a/Engine/Source/Editor/UnrealEd/Private/EdGraphNode_Comment.cpp
+++ b/Engine/Source/Editor/UnrealEd/Private/EdGraphNode_Comment.cpp
@@ -91,9 +91,9 @@ FText UEdGraphNode_Comment::GetNodeTitle(ENodeTitleType::Type TitleType) const
 	return FText::FromString(NodeComment);
 }
 
-FString UEdGraphNode_Comment::GetPinNameOverride(const UEdGraphPin& Pin) const
+FText UEdGraphNode_Comment::GetPinNameOverride(const UEdGraphPin& Pin) const
 {
-	return GetNodeTitle(ENodeTitleType::ListView).ToString();
+	return GetNodeTitle(ENodeTitleType::ListView);
 }
 
 FLinearColor UEdGraphNode_Comment::GetNodeCommentColor() const
diff --git a/Engine/Source/Editor/UnrealEd/Private/Factories/CSVImportFactory.cpp b/Engine/Source/Editor/UnrealEd/Private/Factories/CSVImportFactory.cpp
--- a/Engine/Source/Editor/UnrealEd/Private/Factories/CSVImportFactory.cpp
+++ b/Engine/Source/Editor/UnrealEd/Private/Factories/CSVImportFactory.cpp
@@ -269,7 +269,7 @@ public:
 	TSharedRef<SWidget> MakeImportTypeItemWidget( TSharedPtr<ECSVImportType> Type )
 	{
 		return	SNew(STextBlock)
-				.Text(GetImportTypeText(Type));
+				.Text(FText::FromString(GetImportTypeText(Type)));
 	}
 
 	/** Called to create a widget for each struct */
@@ -277,7 +277,7 @@ public:
 	{
 		check( Struct != NULL );
 		return	SNew(STextBlock)
-				.Text(Struct->GetName());
+				.Text(FText::FromString(Struct->GetName()));
 	}
 
 	FString GetCurveTypeText (CurveInterpModePtr InterpMode) const
@@ -305,7 +305,7 @@ public:
 	TSharedRef<SWidget> MakeCurveTypeWidget( CurveInterpModePtr InterpMode )
 	{
 		FString Label = GetCurveTypeText(InterpMode);
-		return SNew(STextBlock) .Text( Label );
+		return SNew(STextBlock) .Text( FText::FromString(Label) );
 	}
 
 	/** Called when 'OK' button is pressed */
@@ -336,29 +336,29 @@ public:
 		return FReply::Handled();
 	}
 
-	FString GetSelectedItemText() const
+	FText GetSelectedItemText() const
 	{
 		TSharedPtr<ECSVImportType> SelectedType = ImportTypeCombo->GetSelectedItem();
 
 		return (SelectedType.IsValid())
-			? GetImportTypeText(SelectedType)
-			: FString();
+			? FText::FromString(GetImportTypeText(SelectedType))
+			: FText::GetEmpty();
 	}
 
-	FString GetSelectedRowOptionText() const
+	FText GetSelectedRowOptionText() const
 	{
 		UScriptStruct* SelectedScript = RowStructCombo->GetSelectedItem();
 		return (SelectedScript)
-			? SelectedScript->GetName()
-			: FString();
+			? FText::FromString(SelectedScript->GetName())
+			: FText::GetEmpty();
 	}
 
-	FString GetSelectedCurveTypeText() const
+	FText GetSelectedCurveTypeText() const
 	{
 		CurveInterpModePtr CurveModePtr = CurveInterpCombo->GetSelectedItem();
 		return (CurveModePtr.IsValid())
-			? GetCurveTypeText(CurveModePtr)
-			: FString();
+			? FText::FromString(GetCurveTypeText(CurveModePtr))
+			: FText::GetEmpty();
 	}
 };
 
diff --git a/Engine/Source/Editor/UnrealEd/Private/Factories/EditorFactories.cpp b/Engine/Source/Editor/UnrealEd/Private/Factories/EditorFactories.cpp
--- a/Engine/Source/Editor/UnrealEd/Private/Factories/EditorFactories.cpp
+++ b/Engine/Source/Editor/UnrealEd/Private/Factories/EditorFactories.cpp
@@ -7208,13 +7208,13 @@ bool UDataTableFactory::ConfigureProperties()
 
 		TSharedRef<SWidget> MakeRowStructItemWidget(class UScriptStruct* Struct) const
 		{
-			return SNew(STextBlock).Text(Struct ? Struct->GetDisplayNameText().ToString() : FString());
+			return SNew(STextBlock).Text(Struct ? Struct->GetDisplayNameText() : FText::GetEmpty());
 		}
 
-		FString GetSelectedRowOptionText() const
+		FText GetSelectedRowOptionText() const
 		{
 			UScriptStruct* Struct = RowStructCombo.IsValid() ? RowStructCombo->GetSelectedItem() : NULL;
-			return Struct ? Struct->GetDisplayNameText().ToString() : FString();
+			return Struct ? Struct->GetDisplayNameText() : FText::GetEmpty();
 		}
 
 		FReply OnCreate()
diff --git a/Engine/Source/Editor/UnrealEd/Private/Fbx/FbxMainImport.cpp b/Engine/Source/Editor/UnrealEd/Private/Fbx/FbxMainImport.cpp
--- a/Engine/Source/Editor/UnrealEd/Private/Fbx/FbxMainImport.cpp
+++ b/Engine/Source/Editor/UnrealEd/Private/Fbx/FbxMainImport.cpp
@@ -112,7 +112,7 @@ FBXImportOptions* GetImportOptions( UnFbx::FFbxImporter* FbxImporter, UFbxImport
 			SAssignNew(FbxOptionWindow, SFbxOptionWindow)
 			.ImportUI(ImportUI)
 			.WidgetWindow(Window)
-			.FullPath(FullPath)
+			.FullPath(FText::FromString(FullPath))
 			.ForcedImportType( bForceImportType ? TOptional<EFBXImportType>( ImportType ) : TOptional<EFBXImportType>() )
 			.IsObjFormat( bIsObjFormat )
 		);
diff --git a/Engine/Source/Editor/UnrealEd/Private/FbxOptionWindow.h b/Engine/Source/Editor/UnrealEd/Private/FbxOptionWindow.h
--- a/Engine/Source/Editor/UnrealEd/Private/FbxOptionWindow.h
+++ b/Engine/Source/Editor/UnrealEd/Private/FbxOptionWindow.h
@@ -18,7 +18,7 @@ public:
 
 		SLATE_ARGUMENT( UFbxImportUI*, ImportUI )
 		SLATE_ARGUMENT( TSharedPtr<SWindow>, WidgetWindow )
-		SLATE_ARGUMENT( FString, FullPath )
+		SLATE_ARGUMENT( FText, FullPath )
 		SLATE_ARGUMENT( TOptional<EFBXImportType>, ForcedImportType )
 		SLATE_ARGUMENT( bool, IsObjFormat )
 	SLATE_END_ARGS()
diff --git a/Engine/Source/Editor/UnrealEd/Private/Kismet2/BlueprintEditorUtils.cpp b/Engine/Source/Editor/UnrealEd/Private/Kismet2/BlueprintEditorUtils.cpp
--- a/Engine/Source/Editor/UnrealEd/Private/Kismet2/BlueprintEditorUtils.cpp
+++ b/Engine/Source/Editor/UnrealEd/Private/Kismet2/BlueprintEditorUtils.cpp
@@ -5976,7 +5976,7 @@ TSharedRef<SWidget> FBlueprintEditorUtils::ConstructBlueprintParentClassPicker(
 
 	TSharedPtr<FBlueprintReparentFilter> Filter = MakeShareable(new FBlueprintReparentFilter);
 	Options.ClassFilter = Filter;
-	Options.ViewerTitleString = LOCTEXT("ReparentBblueprint", "Reparent blueprint").ToString();
+	Options.ViewerTitleString = LOCTEXT("ReparentBblueprint", "Reparent blueprint");
 
 	// Only allow parenting to base blueprints.
 	Options.bIsBlueprintBaseOnly = true;
@@ -6119,7 +6119,7 @@ TSharedRef<SWidget> FBlueprintEditorUtils::ConstructBlueprintInterfaceClassPicke
 
 	TSharedPtr<FBlueprintInterfaceFilter> Filter = MakeShareable(new FBlueprintInterfaceFilter);
 	Options.ClassFilter = Filter;
-	Options.ViewerTitleString = LOCTEXT("ImplementInterfaceBlueprint", "Implement Interface").ToString();
+	Options.ViewerTitleString = LOCTEXT("ImplementInterfaceBlueprint", "Implement Interface");
 
 	for( auto BlueprintIter = Blueprints.CreateConstIterator(); BlueprintIter; ++BlueprintIter )
 	{
diff --git a/Engine/Source/Editor/UnrealEd/Private/ObjectTools.cpp b/Engine/Source/Editor/UnrealEd/Private/ObjectTools.cpp
--- a/Engine/Source/Editor/UnrealEd/Private/ObjectTools.cpp
+++ b/Engine/Source/Editor/UnrealEd/Private/ObjectTools.cpp
@@ -973,7 +973,7 @@ namespace ObjectTools
 					
 					// Display a dialog containing all referencers; the dialog is designed to destroy itself upon being closed, so this
 					// allocation is ok and not a memory leak
-					SGenericDialogWidget::OpenDialog(NSLOCTEXT("ObjectTools", "ShowReferencers", "Show Referencers"), SNew(STextBlock).Text( Ar ));
+					SGenericDialogWidget::OpenDialog(NSLOCTEXT("ObjectTools", "ShowReferencers", "Show Referencers"), SNew(STextBlock).Text( FText::FromString(Ar) ));
 				}
 				else
 				{
@@ -1079,7 +1079,7 @@ namespace ObjectTools
 
 					// Display the object references in a copy-friendly dialog; the dialog is designed to destroy itself upon being closed, so this
 					// allocation is ok and not a memory leak
-					SGenericDialogWidget::OpenDialog(NSLOCTEXT("ObjectTools", "ShowReferencedAssets", "Show Referenced Assets"), SNew(STextBlock).Text(OutString));
+					SGenericDialogWidget::OpenDialog(NSLOCTEXT("ObjectTools", "ShowReferencedAssets", "Show Referenced Assets"), SNew(STextBlock).Text(FText::FromString(OutString)));
 				}
 				else
 				{
diff --git a/Engine/Source/Editor/UnrealEd/Private/PackageRestore.cpp b/Engine/Source/Editor/UnrealEd/Private/PackageRestore.cpp
--- a/Engine/Source/Editor/UnrealEd/Private/PackageRestore.cpp
+++ b/Engine/Source/Editor/UnrealEd/Private/PackageRestore.cpp
@@ -267,7 +267,7 @@ namespace PackageRestore
 				.FillWidth(1)
 				[
 					SNew(STextBlock) 
-					.Text(Item->GetPackageName())
+					.Text(FText::FromString(Item->GetPackageName()))
 				]
 				+SHorizontalBox::Slot()
 				.Padding(FMargin(4, 0, 0, 0))
@@ -280,7 +280,7 @@ namespace PackageRestore
 					.FillWidth(1)
 					[
 						SNew(STextBlock) 
-						.Text(Item->GetPackageFilename())
+						.Text(FText::FromString(Item->GetPackageFilename()))
 					]
 					+SHorizontalBox::Slot()
 					.Padding(FMargin(2, 0, 0, 0))
@@ -302,7 +302,7 @@ namespace PackageRestore
 					.FillWidth(1)
 					[
 						SNew(STextBlock) 
-						.Text(Item->GetAutoSaveFilename())
+						.Text(FText::FromString(Item->GetAutoSaveFilename()))
 					]
 					+SHorizontalBox::Slot()
 					.Padding(FMargin(2, 0, 0, 0))
diff --git a/Engine/Source/Editor/UnrealEd/Private/SComponentClassCombo.cpp b/Engine/Source/Editor/UnrealEd/Private/SComponentClassCombo.cpp
--- a/Engine/Source/Editor/UnrealEd/Private/SComponentClassCombo.cpp
+++ b/Engine/Source/Editor/UnrealEd/Private/SComponentClassCombo.cpp
@@ -221,7 +221,7 @@ TSharedRef<ITableRow> SComponentClassCombo::GenerateAddComponentRow( FComponentC
 				.Padding(1.f)
 				[
 					SNew(STextBlock)
-					.Text(Entry->GetHeadingText())
+					.Text(FText::FromString(Entry->GetHeadingText()))
 					.TextStyle(FEditorStyle::Get(), TEXT("Menu.Heading"))
 				]
 			];
@@ -320,7 +320,7 @@ TSharedRef<ITableRow> SComponentClassCombo::GenerateAddComponentRow( FComponentC
 				[
 					SNew(STextBlock)
 					.HighlightText(this, &SComponentClassCombo::GetCurrentSearchString)
-					.Text( FriendlyComponentName )
+					.Text( FText::FromString(FriendlyComponentName) )
 				]
 			];
 	}
diff --git a/Engine/Source/Editor/UnrealEd/Private/SCurveEditor.cpp b/Engine/Source/Editor/UnrealEd/Private/SCurveEditor.cpp
--- a/Engine/Source/Editor/UnrealEd/Private/SCurveEditor.cpp
+++ b/Engine/Source/Editor/UnrealEd/Private/SCurveEditor.cpp
@@ -389,7 +389,7 @@ TSharedRef<SWidget> SCurveEditor::CreateCurveSelectionWidget() const
 					SNew(STextBlock)
 					.Font(FEditorStyle::GetFontStyle("CurveEd.LabelFont"))
 					.ColorAndOpacity(CurveViewModel->Color)
-					.Text(CurveViewModel->CurveInfo.CurveName.ToString())
+					.Text(FText::FromName(CurveViewModel->CurveInfo.CurveName))
 				]
 
 				+ SHorizontalBox::Slot()
diff --git a/Engine/Source/Editor/UnrealEd/Private/SKeySelector.cpp b/Engine/Source/Editor/UnrealEd/Private/SKeySelector.cpp
--- a/Engine/Source/Editor/UnrealEd/Private/SKeySelector.cpp
+++ b/Engine/Source/Editor/UnrealEd/Private/SKeySelector.cpp
@@ -38,15 +38,15 @@ public:
 	{
 	}
 
-	FString GetDescription() const
+	FText GetDescription() const
 	{
 		if (Key.IsValid())
 		{
-			return Key->GetDisplayName().ToString();
+			return Key->GetDisplayName();
 		}
 		else
 		{
-			return Name.ToString();
+			return Name;
 		}
 	}
 
@@ -57,7 +57,7 @@ public:
 
 	bool MatchesSearchTokens(const TArray<FString>& SearchTokens)
 	{
-		FString Description = GetDescription();
+		FString Description = GetDescription().ToString();
 
 		for (auto Token : SearchTokens)
 		{
@@ -146,13 +146,13 @@ void SKeySelector::Construct(const FArguments& InArgs)
 //=======================================================================
 // Attribute Helpers
 
-FString SKeySelector::GetKeyDescription() const
+FText SKeySelector::GetKeyDescription() const
 {
 	if (bHasMultipleValues)
 	{
-		return LOCTEXT("MultipleValues", "Multiple Values").ToString();
+		return LOCTEXT("MultipleValues", "Multiple Values");
 	}
-	return CurrentKey.Get().GetDisplayName().ToString();
+	return CurrentKey.Get().GetDisplayName();
 }
 
 const FSlateBrush* SKeySelector::GetKeyIconImage() const
@@ -169,7 +169,7 @@ const FSlateBrush* SKeySelector::GetKeyIconImage() const
 TSharedRef<ITableRow> SKeySelector::GenerateKeyTreeRow(FKeyTreeItem InItem, const TSharedRef<STableViewBase>& OwnerTree)
 {
 	const bool bIsCategory = !InItem->GetKey().IsValid();
-	const FString Description = InItem->GetDescription();
+	const FText Description = InItem->GetDescription();
 
 	// Determine the best icon the to represents this item
 	const FSlateBrush* IconBrush = nullptr;
@@ -179,7 +179,7 @@ TSharedRef<ITableRow> SKeySelector::GenerateKeyTreeRow(FKeyTreeItem InItem, cons
 	}
 
 	return SNew(SComboRow<FKeyTreeItem>, OwnerTree)
-		.ToolTip(IDocumentation::Get()->CreateToolTip(FText::FromString(Description), NULL, *BigTooltipDocLink, *Description))
+		.ToolTip(IDocumentation::Get()->CreateToolTip(Description, nullptr, BigTooltipDocLink, Description.ToString()))
 		[
 			SNew(SHorizontalBox)
 			+ SHorizontalBox::Slot()
diff --git a/Engine/Source/Editor/UnrealEd/Private/SScalabilitySettings.cpp b/Engine/Source/Editor/UnrealEd/Private/SScalabilitySettings.cpp
--- a/Engine/Source/Editor/UnrealEd/Private/SScalabilitySettings.cpp
+++ b/Engine/Source/Editor/UnrealEd/Private/SScalabilitySettings.cpp
@@ -49,9 +49,9 @@ float SScalabilitySettings::GetResolutionScale() const
 	return (float)(CachedQualityLevels.ResolutionQuality - Scalability::MinResolutionScale) / (float)(Scalability::MaxResolutionScale - Scalability::MinResolutionScale);
 }
 
-FString SScalabilitySettings::GetResolutionScaleString() const
+FText SScalabilitySettings::GetResolutionScaleString() const
 {
-	return FString::Printf(TEXT("%i%%"), (int32)(100.0f * FMath::Square((float)CachedQualityLevels.ResolutionQuality / 100.0f)));
+	return FText::AsPercent(FMath::Square((float)CachedQualityLevels.ResolutionQuality / 100.0f));
 }
 
 TSharedRef<SWidget> SScalabilitySettings::MakeButtonWidget(const FText& InName, const TCHAR* InGroupName, int32 InQualityLevel, const FText& InToolTip)
diff --git a/Engine/Source/Editor/UnrealEd/Private/SSkeletonWidget.cpp b/Engine/Source/Editor/UnrealEd/Private/SSkeletonWidget.cpp
--- a/Engine/Source/Editor/UnrealEd/Private/SSkeletonWidget.cpp
+++ b/Engine/Source/Editor/UnrealEd/Private/SSkeletonWidget.cpp
@@ -174,7 +174,7 @@ void SSkeletonCompareWidget::Construct(const FArguments& InArgs)
 				.HAlign(HAlign_Fill)
 				[
 					SNew(STextBlock)
-					.Text(Object->GetFullName())
+					.Text(FText::FromString(Object->GetFullName()))
 				]
 			]
 
diff --git a/Engine/Source/Editor/UnrealEd/Private/SSocketManager.cpp b/Engine/Source/Editor/UnrealEd/Private/SSocketManager.cpp
--- a/Engine/Source/Editor/UnrealEd/Private/SSocketManager.cpp
+++ b/Engine/Source/Editor/UnrealEd/Private/SSocketManager.cpp
@@ -629,7 +629,7 @@ FReply SSocketManager::DeleteSelectedSocket_Execute()
 	return FReply::Handled();
 }
 
-FString SSocketManager::GetSocketHeaderText() const
+FText SSocketManager::GetSocketHeaderText() const
 {
 	UStaticMesh* CurrentStaticMesh = nullptr;
 	TSharedPtr<IStaticMeshEditor> StaticMeshEditorPinned = StaticMeshEditorPtr.Pin();
@@ -637,7 +637,7 @@ FString SSocketManager::GetSocketHeaderText() const
 	{
 		CurrentStaticMesh = StaticMeshEditorPinned->GetStaticMesh();
 	}
-	return FString::Printf(*LOCTEXT("SocketHeader_Total", "Sockets (%d Total)").ToString(), (CurrentStaticMesh != nullptr) ? CurrentStaticMesh->Sockets.Num() : 0);
+	return FText::Format(LOCTEXT("SocketHeader_TotalFmt", "Sockets ({0} Total)"), FText::AsNumber((CurrentStaticMesh != nullptr) ? CurrentStaticMesh->Sockets.Num() : 0));
 }
 
 void SSocketManager::SocketName_TextChanged(const FText& InText)
diff --git a/Engine/Source/Editor/UnrealEd/Private/SSocketManager.h b/Engine/Source/Editor/UnrealEd/Private/SSocketManager.h
--- a/Engine/Source/Editor/UnrealEd/Private/SSocketManager.h
+++ b/Engine/Source/Editor/UnrealEd/Private/SSocketManager.h
@@ -61,7 +61,7 @@ private:
 	/** Callback for the Delete Socket button. */
 	FReply DeleteSelectedSocket_Execute();
 
-	FString GetSocketHeaderText() const;
+	FText GetSocketHeaderText() const;
 
 	/** Callback for when the socket name textbox is changed, verifies the name is not a duplicate. */
 	void SocketName_TextChanged(const FText& InText);
diff --git a/Engine/Source/Editor/UnrealEd/Private/Settings/LevelEditorPlaySettingsCustomization.h b/Engine/Source/Editor/UnrealEd/Private/Settings/LevelEditorPlaySettingsCustomization.h
--- a/Engine/Source/Editor/UnrealEd/Private/Settings/LevelEditorPlaySettingsCustomization.h
+++ b/Engine/Source/Editor/UnrealEd/Private/Settings/LevelEditorPlaySettingsCustomization.h
@@ -498,7 +498,7 @@ public:
 private:
 
 	// Callback for getting the description of the settings
-	FString HandleMultiplayerOptionsDescription( ) const
+	FText HandleMultiplayerOptionsDescription( ) const
 	{
 		const ULevelEditorPlaySettings* PlayInSettings = GetDefault<ULevelEditorPlaySettings>();
 		const bool CanRunUnderOneProcess = [&PlayInSettings]{ bool RunUnderOneProcess(false); return (PlayInSettings->GetRunUnderOneProcess(RunUnderOneProcess) && RunUnderOneProcess); }();
@@ -578,7 +578,7 @@ private:
 				}
 			}
 		}
-		return Desc;
+		return FText::FromString(Desc);
 	}
 
 	// Callback for checking whether the ClientWindowHeight and ClientWindowWidth properties are enabled.
diff --git a/Engine/Source/Editor/UnrealEd/Private/Tests/BlueprintAutomationTests.cpp b/Engine/Source/Editor/UnrealEd/Private/Tests/BlueprintAutomationTests.cpp
--- a/Engine/Source/Editor/UnrealEd/Private/Tests/BlueprintAutomationTests.cpp
+++ b/Engine/Source/Editor/UnrealEd/Private/Tests/BlueprintAutomationTests.cpp
@@ -988,10 +988,10 @@ bool FBlueprintCompileOnLoadTest::RunTest(const FString& BlueprintAssetPath)
 		for (auto DiffIt(BlueprintDiffs.CreateIterator()); DiffIt; ++DiffIt)
 		{
 			// will be presented in the context of "what changed between the initial load and the second?"
-			FString DiffDescription = DiffIt->ToolTip;
-			if (DiffDescription != DiffIt->DisplayString)
+			FString DiffDescription = DiffIt->ToolTip.ToString();
+			if (DiffDescription != DiffIt->DisplayString.ToString())
 			{
-				DiffDescription = FString::Printf(TEXT("%s (%s)"), *DiffDescription, *DiffIt->DisplayString);
+				DiffDescription = FString::Printf(TEXT("%s (%s)"), *DiffDescription, *DiffIt->DisplayString.ToString());
 			}
 
 			const UEdGraphNode* NodeFromPin = DiffIt->Pin1 ? Cast<const UEdGraphNode>(DiffIt->Pin1->GetOuter()) : NULL;
diff --git a/Engine/Source/Editor/UnrealEd/Public/Dialogs/DlgSoundWaveOptions.h b/Engine/Source/Editor/UnrealEd/Public/Dialogs/DlgSoundWaveOptions.h
--- a/Engine/Source/Editor/UnrealEd/Public/Dialogs/DlgSoundWaveOptions.h
+++ b/Engine/Source/Editor/UnrealEd/Public/Dialogs/DlgSoundWaveOptions.h
@@ -39,7 +39,7 @@ private:
 
 	/** Progress Bar Callbacks */
 	TOptional<float> GetProgress() const;
-	FString GetProgressString() const;
+	FText GetProgressString() const;
 	EVisibility IsProgressBarVisible() const;
 
 	/** Button click callbacks. */
diff --git a/Engine/Source/Editor/UnrealEd/Public/DragAndDrop/ClassDragDropOp.h b/Engine/Source/Editor/UnrealEd/Public/DragAndDrop/ClassDragDropOp.h
--- a/Engine/Source/Editor/UnrealEd/Public/DragAndDrop/ClassDragDropOp.h
+++ b/Engine/Source/Editor/UnrealEd/Public/DragAndDrop/ClassDragDropOp.h
@@ -37,7 +37,7 @@ public:
 					.AutoWidth()
 					[
 						SNew(STextBlock) 
-						.Text( ClassesToDrop[0]->GetName() )
+						.Text( FText::FromString(ClassesToDrop[0]->GetName()) )
 					]
 				];
 		}
@@ -51,7 +51,7 @@ public:
 				.AutoWidth()
 				[
 					SNew(STextBlock) 
-					.Text( ClassesToDrop[0]->GetName() )
+					.Text( FText::FromString(ClassesToDrop[0]->GetName()) )
 				]
 			];
 	}
@@ -102,7 +102,7 @@ public:
 				.AutoWidth()
 				[
 					SNew(STextBlock) 
-					.Text( (*AssetsToDrop.Get())[0].AssetName )
+					.Text( FText::FromString((*AssetsToDrop.Get())[0].AssetName) )
 				]
 			];
 	}
diff --git a/Engine/Source/Editor/UnrealEd/Public/DragAndDrop/LevelDragDropOp.h b/Engine/Source/Editor/UnrealEd/Public/DragAndDrop/LevelDragDropOp.h
--- a/Engine/Source/Editor/UnrealEd/Public/DragAndDrop/LevelDragDropOp.h
+++ b/Engine/Source/Editor/UnrealEd/Public/DragAndDrop/LevelDragDropOp.h
@@ -46,7 +46,7 @@ public:
 				.AutoWidth()
 					[
 						SNew(STextBlock) 
-						.Text( LevelName )
+						.Text( FText::FromString(LevelName) )
 					]
 			];
 	}
diff --git a/Engine/Source/Editor/UnrealEd/Public/EdGraphNode_Comment.h b/Engine/Source/Editor/UnrealEd/Public/EdGraphNode_Comment.h
--- a/Engine/Source/Editor/UnrealEd/Public/EdGraphNode_Comment.h
+++ b/Engine/Source/Editor/UnrealEd/Public/EdGraphNode_Comment.h
@@ -53,7 +53,7 @@ public:
 	UNREALED_API virtual FLinearColor GetNodeCommentColor() const override;
 	UNREALED_API virtual FText GetNodeTitle(ENodeTitleType::Type TitleType) const override;
 	virtual bool ShouldOverridePinNames() const override { return true; }
-	UNREALED_API virtual FString GetPinNameOverride(const UEdGraphPin& Pin) const override;
+	UNREALED_API virtual FText GetPinNameOverride(const UEdGraphPin& Pin) const override;
 	UNREALED_API virtual void ResizeNode(const FVector2D& NewSize) override;
 	UNREALED_API virtual void PostPlacedNewNode() override;
 	UNREALED_API virtual void OnRenameNode(const FString& NewName) override;
diff --git a/Engine/Source/Editor/UnrealEd/Public/GraphEditor.h b/Engine/Source/Editor/UnrealEd/Public/GraphEditor.h
--- a/Engine/Source/Editor/UnrealEd/Public/GraphEditor.h
+++ b/Engine/Source/Editor/UnrealEd/Public/GraphEditor.h
@@ -26,11 +26,11 @@ struct FGraphAppearanceInfo
 	/** Image to draw in corner of graph */
 	const FSlateBrush* CornerImage;
 	/** Text to write in corner of graph */
-	FString CornerText;
+	FText CornerText;
 	/** If set, will be used as override for PIE notify text */
-	FString PIENotifyText;
+	FText PIENotifyText;
 	/** If set, will be used as override for read only text */
-	FString ReadOnlyText;
+	FText ReadOnlyText;
 	/** Text to display if the graph is empty (to guide the user on what to do) */
 	FText InstructionText;
 	/** Allows graphs to nicely fade instruction text (or completely hide it). */
diff --git a/Engine/Source/Editor/UnrealEd/Public/SKeySelector.h b/Engine/Source/Editor/UnrealEd/Public/SKeySelector.h
--- a/Engine/Source/Editor/UnrealEd/Public/SKeySelector.h
+++ b/Engine/Source/Editor/UnrealEd/Public/SKeySelector.h
@@ -40,7 +40,7 @@ protected:
 	/** Gets the icon for the key being manipulated */
 	const FSlateBrush* GetKeyIconImage() const;
 	/** Gets a succinct description for the key being manipulated */
-	FString GetKeyDescription() const;
+	FText GetKeyDescription() const;
 
 	/** Treeview support functions */
 	virtual TSharedRef<ITableRow> GenerateKeyTreeRow(FKeyTreeItem InItem, const TSharedRef<STableViewBase>& OwnerTree);
diff --git a/Engine/Source/Editor/UnrealEd/Public/SScalabilitySettings.h b/Engine/Source/Editor/UnrealEd/Public/SScalabilitySettings.h
--- a/Engine/Source/Editor/UnrealEd/Public/SScalabilitySettings.h
+++ b/Engine/Source/Editor/UnrealEd/Public/SScalabilitySettings.h
@@ -36,7 +36,7 @@ private:
 	float GetResolutionScale() const;
 
 	// Callback to retrieve current resolution scale as a display string
-	FString GetResolutionScaleString() const;
+	FText GetResolutionScaleString() const;
 
 	// Makes a button widget for the group quality levels
 	TSharedRef<SWidget> MakeButtonWidget(const FText& InName, const TCHAR* InGroupName, int32 InQualityLevel, const FText& InToolTip);
diff --git a/Engine/Source/Editor/UnrealEd/Public/SSkeletonWidget.h b/Engine/Source/Editor/UnrealEd/Public/SSkeletonWidget.h
--- a/Engine/Source/Editor/UnrealEd/Public/SSkeletonWidget.h
+++ b/Engine/Source/Editor/UnrealEd/Public/SSkeletonWidget.h
@@ -56,7 +56,7 @@ public:
 				.Content()
 				[
 					SNew(STextBlock)
-					.Text(BonePair->Bone1.ToString())
+					.Text(FText::FromName(BonePair->Bone1))
 				];
 		}
 		else 
@@ -77,7 +77,7 @@ public:
 					.Content()
 					[
 						SNew(STextBlock)
-						.Text(BonePair->Bone2.ToString())
+						.Text(FText::FromName(BonePair->Bone2))
 					];
 			}
 		}
@@ -118,7 +118,7 @@ public:
 			. Content()
 			[
 				SNew(STextBlock)
-				.Text(InSkeleton->GetFullName())	
+				.Text(FText::FromString(InSkeleton->GetFullName()))
 			];
 	}
 
@@ -130,7 +130,7 @@ public:
 			. Content()
 			[
 				SNew(STextBlock)
-				.Text(InBoneName->ToString())	
+				.Text(FText::FromName(*InBoneName))
 			];
 	}
 
@@ -500,7 +500,7 @@ public:
 		. Content()
 		[
 			SNew(STextBlock)
-			.Text(InBoneName->ToString())	
+			.Text(FText::FromName(*InBoneName))
 		];
 	}
 
diff --git a/Engine/Source/Editor/WorldBrowser/Private/SPropertyEditorLevelPackage.cpp b/Engine/Source/Editor/WorldBrowser/Private/SPropertyEditorLevelPackage.cpp
--- a/Engine/Source/Editor/WorldBrowser/Private/SPropertyEditorLevelPackage.cpp
+++ b/Engine/Source/Editor/WorldBrowser/Private/SPropertyEditorLevelPackage.cpp
@@ -109,7 +109,7 @@ TSharedRef<SWidget> SPropertyEditorLevelPackage::MakePickerWidget()
 				.AutoHeight()
 				[
 					SNew( SSearchBox )
-					.ToolTipText(LOCTEXT("LevelPackage_FilterTooltip", "Type here to search levels").ToString())
+					.ToolTipText(LOCTEXT("LevelPackage_FilterTooltip", "Type here to search levels"))
 					.HintText(LOCTEXT( "LevelPackage_FilterHint", "Search Levels" ))
 					.OnTextChanged(SearchBoxLevelPackageFilter.Get(), &LevelPackageTextFilter::SetRawFilterText)
 				]
@@ -149,8 +149,8 @@ TSharedRef<ITableRow> SPropertyEditorLevelPackage::MakeListRowWidget(TSharedPtr<
 			.BorderImage(FEditorStyle::GetBrush("NoBrush"))
 			[
 				SNew(STextBlock)
-				.ToolTipText(InPackageItem->LongPackageName)
-				.Text(InPackageItem->DisplayName)
+				.ToolTipText(FText::FromString(InPackageItem->LongPackageName))
+				.Text(FText::FromString(InPackageItem->DisplayName))
 				
 			]
 		];
diff --git a/Engine/Source/Editor/WorldBrowser/Private/SWorldDetails.cpp b/Engine/Source/Editor/WorldBrowser/Private/SWorldDetails.cpp
--- a/Engine/Source/Editor/WorldBrowser/Private/SWorldDetails.cpp
+++ b/Engine/Source/Editor/WorldBrowser/Private/SWorldDetails.cpp
@@ -212,7 +212,7 @@ TSharedRef<SWidget> SWorldDetails::HandleInspectedLevelComboBoxGenerateWidget(TS
 	return SNew(SBox)
 	.Padding(4)
 	[
-		SNew(STextBlock).Text(InLevelModel->GetDisplayName())
+		SNew(STextBlock).Text(FText::FromString(InLevelModel->GetDisplayName()))
 	];
 }
 
diff --git a/Engine/Source/Editor/WorldBrowser/Private/SWorldHierarchy.cpp b/Engine/Source/Editor/WorldBrowser/Private/SWorldHierarchy.cpp
--- a/Engine/Source/Editor/WorldBrowser/Private/SWorldHierarchy.cpp
+++ b/Engine/Source/Editor/WorldBrowser/Private/SWorldHierarchy.cpp
@@ -203,7 +203,7 @@ public:
 			.AutoHeight()
 			[
 				SNew( SSearchBox )
-				.ToolTipText(LOCTEXT("FilterSearchToolTip", "Type here to search Levels").ToString())
+				.ToolTipText(LOCTEXT("FilterSearchToolTip", "Type here to search Levels"))
 				.HintText(LOCTEXT("FilterSearchHint", "Search Levels"))
 				.OnTextChanged(SearchBoxLevelFilter.Get(), &LevelTextFilter::SetRawFilterText)
 			]
@@ -443,7 +443,7 @@ private:
 	}
 
 	/** @return	Returns the filter status text */
-	FString GetFilterStatusText() const
+	FText GetFilterStatusText() const
 	{
 		const int32 SelectedLevelsCount = WorldModel->GetSelectedLevels().Num();
 		const int32 TotalLevelsCount = WorldModel->GetAllLevels().Num();
@@ -453,24 +453,24 @@ private:
 		{
 			if (SelectedLevelsCount == 0)
 			{
-				return FString::Printf( *LOCTEXT("ShowingAllLevels", "%d levels").ToString(), TotalLevelsCount );
+				return FText::Format( LOCTEXT("ShowingAllLevelsFmt", "{0} levels"), FText::AsNumber(TotalLevelsCount) );
 			}
 			else
 			{
-				return FString::Printf( *LOCTEXT("ShowingAllLevelsSelected", "%d levels (%d selected)").ToString(), TotalLevelsCount, SelectedLevelsCount );
+				return FText::Format( LOCTEXT("ShowingAllLevelsSelectedFmt", "{0} levels ({1} selected)"), FText::AsNumber(TotalLevelsCount), FText::AsNumber(SelectedLevelsCount) );
 			}
 		}
 		else if(WorldModel->IsFilterActive() && FilteredLevelsCount == 0)
 		{
-			return FString::Printf( *LOCTEXT("ShowingNoLevels", "No matching levels (%d total)").ToString(), TotalLevelsCount );
+			return FText::Format( LOCTEXT("ShowingNoLevelsFmt", "No matching levels ({0} total)"), FText::AsNumber(TotalLevelsCount) );
 		}
 		else if (SelectedLevelsCount != 0)
 		{
-			return FString::Printf( *LOCTEXT("ShowingOnlySomeLevelsSelected", "Showing %d of %d levels (%d selected)").ToString(), FilteredLevelsCount, TotalLevelsCount, SelectedLevelsCount );
+			return FText::Format( LOCTEXT("ShowingOnlySomeLevelsSelectedFmt", "Showing {0} of {1} levels ({2} selected)"), FText::AsNumber(FilteredLevelsCount), FText::AsNumber(TotalLevelsCount), FText::AsNumber(SelectedLevelsCount) );
 		}
 		else
 		{
-			return FString::Printf( *LOCTEXT("ShowingOnlySomeLevels", "Showing %d of %d levels").ToString(), FilteredLevelsCount, TotalLevelsCount );
+			return FText::Format( LOCTEXT("ShowingOnlySomeLevelsFmt", "Showing {0} of {1} levels"), FText::AsNumber(FilteredLevelsCount), FText::AsNumber(TotalLevelsCount) );
 		}
 	}
 
diff --git a/Engine/Source/Editor/WorldBrowser/Private/SWorldHierarchyItem.cpp b/Engine/Source/Editor/WorldBrowser/Private/SWorldHierarchyItem.cpp
--- a/Engine/Source/Editor/WorldBrowser/Private/SWorldHierarchyItem.cpp
+++ b/Engine/Source/Editor/WorldBrowser/Private/SWorldHierarchyItem.cpp
@@ -198,9 +198,9 @@ TSharedRef< SWidget > SWorldHierarchyItem::GenerateWidgetForColumn( const FName&
 }
 END_SLATE_FUNCTION_BUILD_OPTIMIZATION
 
-FString SWorldHierarchyItem::GetLevelDisplayNameText() const
+FText SWorldHierarchyItem::GetLevelDisplayNameText() const
 {
-	return LevelModel->GetDisplayName();
+	return FText::FromString(LevelModel->GetDisplayName());
 }
 
 bool SWorldHierarchyItem::IsSaveEnabled() const
diff --git a/Engine/Source/Editor/WorldBrowser/Private/SWorldHierarchyItem.h b/Engine/Source/Editor/WorldBrowser/Private/SWorldHierarchyItem.h
--- a/Engine/Source/Editor/WorldBrowser/Private/SWorldHierarchyItem.h
+++ b/Engine/Source/Editor/WorldBrowser/Private/SWorldHierarchyItem.h
@@ -104,7 +104,7 @@ private:
 	/**
 	 *  @return The text of level name while it is not being edited
 	 */
-	FString GetLevelDisplayNameText() const;
+	FText GetLevelDisplayNameText() const;
 	
 	/** */
 	FSlateFontInfo GetLevelDisplayNameFont() const;
diff --git a/Engine/Source/Editor/WorldBrowser/Private/Tiles/STiledLandscapeImportDlg.cpp b/Engine/Source/Editor/WorldBrowser/Private/Tiles/STiledLandscapeImportDlg.cpp
--- a/Engine/Source/Editor/WorldBrowser/Private/Tiles/STiledLandscapeImportDlg.cpp
+++ b/Engine/Source/Editor/WorldBrowser/Private/Tiles/STiledLandscapeImportDlg.cpp
@@ -318,7 +318,7 @@ TSharedRef<ITableRow> STiledLandcapeImportDlg::OnGenerateWidgetForLayerDataListV
 					.HAlign(HAlign_Left)
 					.FillWidth(1)
 					[
-						SNew(STextBlock).Text(InLayerData->Name.ToString())
+						SNew(STextBlock).Text(FText::FromName(InLayerData->Name))
 					]
 
 					// Blend option
diff --git a/Engine/Source/Editor/WorldBrowser/Private/Tiles/SWorldComposition.cpp b/Engine/Source/Editor/WorldBrowser/Private/Tiles/SWorldComposition.cpp
--- a/Engine/Source/Editor/WorldBrowser/Private/Tiles/SWorldComposition.cpp
+++ b/Engine/Source/Editor/WorldBrowser/Private/Tiles/SWorldComposition.cpp
@@ -1162,7 +1162,7 @@ TSharedRef<SWidget> SWorldComposition::ConstructContentWidget()
 				SNew(STextBlock)
 				.Visibility(this, &SWorldComposition::IsSimulationVisible)
 				.TextStyle( FEditorStyle::Get(), "Graph.SimulatingText" )
-				.Text(FString(TEXT("SIMULATING")))
+				.Text(LOCTEXT("SimulatingNotification", "SIMULATING"))
 			]
 		];
 }
@@ -1238,63 +1238,58 @@ FText SWorldComposition::GetZoomText() const
 	return GridView->GetZoomText();
 }
 
-FString SWorldComposition::GetCurrentOriginText() const
+FText SWorldComposition::GetCurrentOriginText() const
 {
 	UWorld* CurrentWorld = (TileWorldModel->IsSimulating() ? TileWorldModel->GetSimulationWorld() : TileWorldModel->GetWorld());
-	return FString::Printf(TEXT("%d, %d"), CurrentWorld->OriginLocation.X, CurrentWorld->OriginLocation.Y);
+	return FText::Format(LOCTEXT("PositionXYFmt", "{0}, {1}"), FText::AsNumber(CurrentWorld->OriginLocation.X), FText::AsNumber(CurrentWorld->OriginLocation.Y));
 }
 
-FString SWorldComposition::GetCurrentLevelText() const
+FText SWorldComposition::GetCurrentLevelText() const
 {
-	FString Text("");
 	UWorld* CurrentWorld = (TileWorldModel->IsSimulating() ? TileWorldModel->GetSimulationWorld() : TileWorldModel->GetWorld());
 
 	if (CurrentWorld->GetCurrentLevel())
 	{
 		UPackage* Package = CastChecked<UPackage>(CurrentWorld->GetCurrentLevel()->GetOutermost());
-		Text+= FPackageName::GetShortName(Package->GetName());
-	}
-	else
-	{
-		Text+= TEXT("None");
+		return FText::FromString(FPackageName::GetShortName(Package->GetName()));
 	}
 	
-	return Text;
+	return LOCTEXT("None", "None");
 }
 
-FString SWorldComposition::GetMouseLocationText() const
+FText SWorldComposition::GetMouseLocationText() const
 {
 	FVector2D MouseLocation = GridView->GetMouseWorldLocation();
-	return FString::Printf(TEXT("%d, %d"), FMath::RoundToInt(MouseLocation.X), FMath::RoundToInt(MouseLocation.Y));
+	return FText::Format(LOCTEXT("PositionXYFmt", "{0}, {1}"), FText::AsNumber(FMath::RoundToInt(MouseLocation.X)), FText::AsNumber(FMath::RoundToInt(MouseLocation.Y)));
 }
 
-FString	SWorldComposition::GetMarqueeSelectionSizeText() const
+FText SWorldComposition::GetMarqueeSelectionSizeText() const
 {
 	FVector2D MarqueeSize = GridView->GetMarqueeWorldSize();
 	
 	if (MarqueeSize.X > 0 && 
 		MarqueeSize.Y > 0)
 	{
-		return FString::Printf(TEXT("%d x %d"), FMath::RoundToInt(MarqueeSize.X), FMath::RoundToInt(MarqueeSize.Y));
+		return FText::Format(LOCTEXT("SizeXYFmt", "{0} x {1}"), FText::AsNumber(FMath::RoundToInt(MarqueeSize.X)), FText::AsNumber(FMath::RoundToInt(MarqueeSize.Y)));
 	}
 	else
 	{
-		return FString();
+		return FText::GetEmpty();
 	}
 }
 
-FString SWorldComposition::GetWorldSizeText() const
+FText SWorldComposition::GetWorldSizeText() const
 {
 	FIntPoint WorldSize = TileWorldModel->GetWorldSize();
 	
 	if (WorldSize.X > 0 && 
 		WorldSize.Y > 0)
 	{
-		return FString::Printf(TEXT("%d x %d"), WorldSize.X, WorldSize.Y);
+		return FText::Format(LOCTEXT("SizeXYFmt", "{0} x {1}"), FText::AsNumber(WorldSize.X), FText::AsNumber(WorldSize.Y));
 	}
 	else
 	{
-		return FString();
+		return FText::GetEmpty();
 	}
 }
 
diff --git a/Engine/Source/Editor/WorldBrowser/Private/Tiles/SWorldComposition.h b/Engine/Source/Editor/WorldBrowser/Private/Tiles/SWorldComposition.h
--- a/Engine/Source/Editor/WorldBrowser/Private/Tiles/SWorldComposition.h
+++ b/Engine/Source/Editor/WorldBrowser/Private/Tiles/SWorldComposition.h
@@ -34,13 +34,13 @@ private:
 
 	/** Top status bar details */
 	FText GetZoomText() const;
-	FString	GetCurrentOriginText() const;
-	FString GetCurrentLevelText() const;
+	FText GetCurrentOriginText() const;
+	FText GetCurrentLevelText() const;
 
 	/** Bottom status bar details */
-	FString GetMouseLocationText() const;
-	FString	GetMarqueeSelectionSizeText() const;
-	FString GetWorldSizeText() const;
+	FText GetMouseLocationText() const;
+	FText GetMarqueeSelectionSizeText() const;
+	FText GetWorldSizeText() const;
 
 	/** @return whether SIMULATION sign should be visible */
 	EVisibility IsSimulationVisible() const;
diff --git a/Engine/Source/Editor/WorldBrowser/Private/Tiles/SWorldLayers.cpp b/Engine/Source/Editor/WorldBrowser/Private/Tiles/SWorldLayers.cpp
--- a/Engine/Source/Editor/WorldBrowser/Private/Tiles/SWorldLayers.cpp
+++ b/Engine/Source/Editor/WorldBrowser/Private/Tiles/SWorldLayers.cpp
@@ -198,7 +198,7 @@ void SWorldLayerButton::Construct(const FArguments& InArgs)
 							SNew(STextBlock)
 								.Font(FEditorStyle::GetFontStyle("ContentBrowser.FilterNameFont"))
 								.ShadowOffset(FVector2D(1.f, 1.f))
-								.Text(WorldLayer.Name)
+								.Text(FText::FromString(WorldLayer.Name))
 						]
 				]
 		];
diff --git a/Engine/Source/Editor/WorldBrowser/Private/Tiles/SWorldTileItem.cpp b/Engine/Source/Editor/WorldBrowser/Private/Tiles/SWorldTileItem.cpp
--- a/Engine/Source/Editor/WorldBrowser/Private/Tiles/SWorldTileItem.cpp
+++ b/Engine/Source/Editor/WorldBrowser/Private/Tiles/SWorldTileItem.cpp
@@ -450,21 +450,21 @@ FReply SWorldTileItem::OnMouseButtonDoubleClick(const FGeometry& InMyGeometry, c
 	return FReply::Handled();
 }
 
-FString SWorldTileItem::GetLevelNameText() const
+FText SWorldTileItem::GetLevelNameText() const
 {
-	return TileModel->GetDisplayName();
+	return FText::FromString(TileModel->GetDisplayName());
 }
 
 FText SWorldTileItem::GetPositionText() const
 {
 	FIntPoint Position = TileModel->GetRelativeLevelPosition();
-	return FText::FromString(FString::Printf(TEXT("%d, %d"), Position.X, Position.Y));
+	return FText::Format(LOCTEXT("PositionXYFmt", "{0}, {1}"), FText::AsNumber(Position.X), FText::AsNumber(Position.Y));
 }
 
 FText SWorldTileItem::GetBoundsExtentText() const
 {
 	FVector2D Size = TileModel->GetLevelSize2D();
-	return FText::FromString(FString::Printf(TEXT("%d, %d"), FMath::RoundToInt(Size.X*0.5f), FMath::RoundToInt(Size.Y*0.5f)));
+	return FText::Format(LOCTEXT("PositionXYFmt", "{0}, {1}"), FText::AsNumber(FMath::RoundToInt(Size.X*0.5f)), FText::AsNumber(FMath::RoundToInt(Size.Y*0.5f)));
 }
 
 FText SWorldTileItem::GetLevelLayerNameText() const
diff --git a/Engine/Source/Editor/WorldBrowser/Private/Tiles/SWorldTileItem.h b/Engine/Source/Editor/WorldBrowser/Private/Tiles/SWorldTileItem.h
--- a/Engine/Source/Editor/WorldBrowser/Private/Tiles/SWorldTileItem.h
+++ b/Engine/Source/Editor/WorldBrowser/Private/Tiles/SWorldTileItem.h
@@ -96,7 +96,7 @@ private:
 	TSharedRef<SToolTip> CreateToolTipWidget();
 	
 	/** Tile tooltips fields */
-	FString GetLevelNameText() const;
+	FText GetLevelNameText() const;
 	FText GetPositionText() const;
 	FText GetBoundsExtentText() const;
 	FText GetLevelLayerNameText() const;
diff --git a/Engine/Source/Runtime/Engine/Classes/EdGraph/EdGraphNode.h b/Engine/Source/Runtime/Engine/Classes/EdGraph/EdGraphNode.h
--- a/Engine/Source/Runtime/Engine/Classes/EdGraph/EdGraphNode.h
+++ b/Engine/Source/Runtime/Engine/Classes/EdGraph/EdGraphNode.h
@@ -190,10 +190,10 @@ public:
 	virtual bool AllowSplitPins() const { return false; }
 
 	/** Gets the overridden name for the specified pin, if any */
-	virtual FString GetPinNameOverride(const UEdGraphPin& Pin) const { return FString(TEXT("")); }
+	virtual FText GetPinNameOverride(const UEdGraphPin& Pin) const { return FText::GetEmpty(); }
 
 	/** Gets the display name for a pin */
-	virtual FString GetPinDisplayName(const UEdGraphPin* Pin) const;
+	virtual FText GetPinDisplayName(const UEdGraphPin* Pin) const;
 
 	/**
 	 * Fetch the hover text for a pin when the graph is being edited.
diff --git a/Engine/Source/Runtime/Engine/Classes/EdGraph/EdGraphNode_Documentation.h b/Engine/Source/Runtime/Engine/Classes/EdGraph/EdGraphNode_Documentation.h
--- a/Engine/Source/Runtime/Engine/Classes/EdGraph/EdGraphNode_Documentation.h
+++ b/Engine/Source/Runtime/Engine/Classes/EdGraph/EdGraphNode_Documentation.h
@@ -25,7 +25,7 @@ public:
 	virtual FLinearColor GetNodeTitleColor() const override { return FLinearColor( 0.2f, 1.f, 0.2f ); }
 	ENGINE_API virtual FText GetTooltipText() const override;
 	ENGINE_API virtual FText GetNodeTitle(ENodeTitleType::Type TitleType) const override;
-	ENGINE_API virtual FString GetPinNameOverride(const UEdGraphPin& Pin) const override;
+	ENGINE_API virtual FText GetPinNameOverride(const UEdGraphPin& Pin) const override;
 	ENGINE_API virtual void ResizeNode(const FVector2D& NewSize) override;
 	ENGINE_API virtual void PostPlacedNewNode() override;
 	ENGINE_API virtual void OnRenameNode(const FString& NewName) override;
diff --git a/Engine/Source/Runtime/Engine/Classes/EdGraph/EdGraphSchema.h b/Engine/Source/Runtime/Engine/Classes/EdGraph/EdGraphSchema.h
--- a/Engine/Source/Runtime/Engine/Classes/EdGraph/EdGraphSchema.h
+++ b/Engine/Source/Runtime/Engine/Classes/EdGraph/EdGraphSchema.h
@@ -204,44 +204,44 @@ struct FEdGraphSchemaAction_Dummy : public FEdGraphSchemaAction
 struct FPinConnectionResponse
 {
 public:
-	FString Message;
+	FText Message;
 
 	TEnumAsByte<enum ECanCreateConnectionResponse> Response;
 
 public:
 	FPinConnectionResponse()
-	: Message(TEXT(""))
+	: Message(FText::GetEmpty())
 	, Response(CONNECT_RESPONSE_MAKE)
 	{
 	}
 
 	FPinConnectionResponse(const ECanCreateConnectionResponse InResponse, const FString& InMessage)
-		: Message(InMessage)
+		: Message(FText::FromString(InMessage))
 		, Response(InResponse)
 	{
 	}
 
 	FPinConnectionResponse(const ECanCreateConnectionResponse InResponse, const ANSICHAR* InMessage)
-		: Message(InMessage)
+		: Message(FText::FromString(InMessage))
 		, Response(InResponse)
 	{
 	}
 
 	FPinConnectionResponse(const ECanCreateConnectionResponse InResponse, const WIDECHAR* InMessage)
-		: Message(InMessage)
+		: Message(FText::FromString(InMessage))
 		, Response(InResponse)
 	{
 	}
 
 	FPinConnectionResponse(const ECanCreateConnectionResponse InResponse, const FText& InMessage)
-		: Message(InMessage.ToString())
+		: Message(InMessage)
 		, Response(InResponse)
 	{
 	}
 
 	friend bool operator==(const FPinConnectionResponse& A, const FPinConnectionResponse& B)
 	{
-		return (A.Message == B.Message) && (A.Response == B.Response);
+		return (A.Message.ToString() == B.Message.ToString()) && (A.Response == B.Response);
 	}	
 
 	/** If a connection can be made without breaking existing connections */
@@ -562,7 +562,7 @@ class ENGINE_API UEdGraphSchema : public UObject
 	virtual FLinearColor GetPinTypeColor(const FEdGraphPinType& PinType) const { return FLinearColor::Black; }
 
 	/** Get the name to show in the editor */
-	virtual FString GetPinDisplayName(const UEdGraphPin* Pin) const;
+	virtual FText GetPinDisplayName(const UEdGraphPin* Pin) const;
 
 	/**
 	 * Takes the PinDescription and tacks on any other data important to the 
diff --git a/Engine/Source/Runtime/Engine/Private/EdGraph/EdGraphNode.cpp b/Engine/Source/Runtime/Engine/Private/EdGraph/EdGraphNode.cpp
--- a/Engine/Source/Runtime/Engine/Private/EdGraph/EdGraphNode.cpp
+++ b/Engine/Source/Runtime/Engine/Private/EdGraph/EdGraphNode.cpp
@@ -315,7 +315,7 @@ UObject* UEdGraphNode::GetJumpTargetForDoubleClick() const
 	return NULL;
 }
 
-FString UEdGraphNode::GetPinDisplayName(const UEdGraphPin* Pin) const
+FText UEdGraphNode::GetPinDisplayName(const UEdGraphPin* Pin) const
 {
 	return GetSchema()->GetPinDisplayName(Pin);
 }
diff --git a/Engine/Source/Runtime/Engine/Private/EdGraph/EdGraphNode_Documentation.cpp b/Engine/Source/Runtime/Engine/Private/EdGraph/EdGraphNode_Documentation.cpp
--- a/Engine/Source/Runtime/Engine/Private/EdGraph/EdGraphNode_Documentation.cpp
+++ b/Engine/Source/Runtime/Engine/Private/EdGraph/EdGraphNode_Documentation.cpp
@@ -50,9 +50,9 @@ FText UEdGraphNode_Documentation::GetNodeTitle( ENodeTitleType::Type TitleType )
 	return NodeTitle;
 }
 
-FString UEdGraphNode_Documentation::GetPinNameOverride(const UEdGraphPin& Pin) const
+FText UEdGraphNode_Documentation::GetPinNameOverride(const UEdGraphPin& Pin) const
 {
-	return GetNodeTitle( ENodeTitleType::ListView ).ToString();
+	return GetNodeTitle( ENodeTitleType::ListView );
 }
 
 void UEdGraphNode_Documentation::ResizeNode(const FVector2D& NewSize)
diff --git a/Engine/Source/Runtime/Engine/Private/EdGraph/EdGraphPin.cpp b/Engine/Source/Runtime/Engine/Private/EdGraph/EdGraphPin.cpp
--- a/Engine/Source/Runtime/Engine/Private/EdGraph/EdGraphPin.cpp
+++ b/Engine/Source/Runtime/Engine/Private/EdGraph/EdGraphPin.cpp
@@ -231,17 +231,17 @@ FString UEdGraphPin::GetDefaultAsString() const
 
 FText UEdGraphPin::GetDisplayName() const
 {
-	FString StrName;
+	FText DisplayName = FText::GetEmpty();
 	auto Schema = GetSchema();
 	if (Schema)
 	{
-		StrName = Schema->GetPinDisplayName(this);
+		DisplayName = Schema->GetPinDisplayName(this);
 	}
 	else
 	{
-		StrName = (!PinFriendlyName.IsEmpty()) ? PinFriendlyName.ToString() : PinName;
+		DisplayName = (!PinFriendlyName.IsEmpty()) ? PinFriendlyName : FText::FromString(PinName);
 	}
-	return FText::FromString(StrName);
+	return DisplayName;
 }
 
 const FString UEdGraphPin::GetLinkInfoString( const FString& InFunctionName, const FString& InInfoData, const UEdGraphPin* InToPin ) const
diff --git a/Engine/Source/Runtime/Engine/Private/EdGraph/EdGraphSchema.cpp b/Engine/Source/Runtime/Engine/Private/EdGraph/EdGraphSchema.cpp
--- a/Engine/Source/Runtime/Engine/Private/EdGraph/EdGraphSchema.cpp
+++ b/Engine/Source/Runtime/Engine/Private/EdGraph/EdGraphSchema.cpp
@@ -489,10 +489,10 @@ FPinConnectionResponse UEdGraphSchema::CopyPinLinks(UEdGraphPin& CopyFromPin, UE
 	return FinalResponse;
 }
 
-FString UEdGraphSchema::GetPinDisplayName(const UEdGraphPin* Pin) const
+FText UEdGraphSchema::GetPinDisplayName(const UEdGraphPin* Pin) const
 {
 	check(Pin != NULL);
-	return !Pin->PinFriendlyName.IsEmpty() ? Pin->PinFriendlyName.ToString() : Pin->PinName;
+	return !Pin->PinFriendlyName.IsEmpty() ? Pin->PinFriendlyName : FText::FromString(Pin->PinName);
 }
 
 void UEdGraphSchema::ConstructBasicPinTooltip(UEdGraphPin const& Pin, FText const& PinDescription, FString& TooltipOut) const
diff --git a/Engine/Source/Runtime/Engine/Private/UserInterface/VisualizeRT.cpp b/Engine/Source/Runtime/Engine/Private/UserInterface/VisualizeRT.cpp
--- a/Engine/Source/Runtime/Engine/Private/UserInterface/VisualizeRT.cpp
+++ b/Engine/Source/Runtime/Engine/Private/UserInterface/VisualizeRT.cpp
@@ -164,27 +164,27 @@ struct FVisualizeRTWidget : public SCompoundWidget, public FReferenceToRenderer
 					SNew( SVerticalBox )
 					+SVerticalBox::Slot()
 					[
-						SNew(STextBlock).Text(FString(TEXT("TODO: R G B A")))
+						SNew(STextBlock).Text(FText::FromString(TEXT("TODO: R G B A")))
 					]
 					+SVerticalBox::Slot()
 					[
-						SNew(STextBlock).Text(FString(TEXT("TODO: Mip")))
+						SNew(STextBlock).Text(FText::FromString(TEXT("TODO: Mip")))
 					]
 					+SVerticalBox::Slot()
 					[
-						SNew(STextBlock).Text(FString(TEXT("TODO: Face")))
+						SNew(STextBlock).Text(FText::FromString(TEXT("TODO: Face")))
 					]
 					+SVerticalBox::Slot()
 					[
-						SNew(STextBlock).Text(FString(TEXT("TODO: Luminance")))
+						SNew(STextBlock).Text(FText::FromString(TEXT("TODO: Luminance")))
 					]
 					+SVerticalBox::Slot()
 					[
-						SNew(STextBlock).Text(FString(TEXT("TODO: Color Multiplier")))
+						SNew(STextBlock).Text(FText::FromString(TEXT("TODO: Color Multiplier")))
 					]
 					+SVerticalBox::Slot()
 					[
-						SNew(STextBlock).Text(FString(TEXT("TODO: Alpha Multiplier")))
+						SNew(STextBlock).Text(FText::FromString(TEXT("TODO: Alpha Multiplier")))
 					]
 				]
 				+ SSplitter::Slot().Value(4)
@@ -267,42 +267,42 @@ struct FVisualizeRTWidget : public SCompoundWidget, public FReferenceToRenderer
 
 		TSharedRef<SWidget> GenerateWidgetForColumn( const FName& Column )
 		{
-			FString Label;
+			FText Label;
 			if ( Column == ColumnName )
 			{
-				Label = Info->Name;
+				Label = FText::FromString(Info->Name);
 			}
 			else if ( Column == ColumnWidth )
 			{
-				Label = Info->Width;
+				Label = FText::FromString(Info->Width);
 			}
 			else if ( Column == ColumnDepth )
 			{
-				Label = Info->Depth;
+				Label = FText::FromString(Info->Depth);
 			}
 			else if ( Column == ColumnHeight )
 			{
-				Label = Info->Height;
+				Label = FText::FromString(Info->Height);
 			}
 			else if ( Column == ColumnFormat )
 			{
-				Label = Info->Format;
+				Label = FText::FromString(Info->Format);
 			}
 			else if ( Column == ColumnDimensions )
 			{
-				Label = Info->Dimensions;
+				Label = FText::FromString(Info->Dimensions);
 			}
 			else if ( Column == ColumnSize )
 			{
-				Label = Info->Size;
+				Label = FText::FromString(Info->Size);
 			}
 			else if ( Column == ColumnType )
 			{
-				Label = Info->Type;
+				Label = FText::FromString(Info->Type);
 			}
 			else if ( Column == ColumnNumber )
 			{
-				Label = Info->Number;
+				Label = FText::FromString(Info->Number);
 			}
 
 			if (!Label.IsEmpty())
diff --git a/Engine/Source/Runtime/GameMenuBuilder/Private/SGameMenuPageWidget.cpp b/Engine/Source/Runtime/GameMenuBuilder/Private/SGameMenuPageWidget.cpp
--- a/Engine/Source/Runtime/GameMenuBuilder/Private/SGameMenuPageWidget.cpp
+++ b/Engine/Source/Runtime/GameMenuBuilder/Private/SGameMenuPageWidget.cpp
@@ -282,9 +282,9 @@ EVisibility SGameMenuPageWidget::GetMenuTitleVisibility() const
 	return CurrentMenuTitle.IsEmpty() || bMenuHidden == true ? EVisibility::Collapsed : EVisibility::Visible;
 }
 
-FString SGameMenuPageWidget::GetMenuTitle() const 
+FText SGameMenuPageWidget::GetMenuTitle() const 
 {
-	return CurrentMenuTitle.ToString();
+	return CurrentMenuTitle;
 }
 
 void SGameMenuPageWidget::SetCurrentMenu(TSharedPtr< class FGameMenuPage > InMenu)
diff --git a/Engine/Source/Runtime/GameMenuBuilder/Public/SGameMenuPageWidget.h b/Engine/Source/Runtime/GameMenuBuilder/Public/SGameMenuPageWidget.h
--- a/Engine/Source/Runtime/GameMenuBuilder/Public/SGameMenuPageWidget.h
+++ b/Engine/Source/Runtime/GameMenuBuilder/Public/SGameMenuPageWidget.h
@@ -253,7 +253,7 @@ protected:
 	FString GetOptionText(TSharedPtr<class FGameMenuItem> MenuItem) const;
 
 	/** Gets current menu title string. */
-	FString GetMenuTitle() const; 
+	FText GetMenuTitle() const; 
 
 	/** Gets the menu title visibility based on if the title text is empty. */
 	EVisibility GetMenuTitleVisibility() const;
