Moved Source Control actions button into level toolbar.

diff --git a/Engine/Source/Editor/EditorStyle/Private/SlateEditorStyle.cpp b/Engine/Source/Editor/EditorStyle/Private/SlateEditorStyle.cpp
--- a/Engine/Source/Editor/EditorStyle/Private/SlateEditorStyle.cpp
+++ b/Engine/Source/Editor/EditorStyle/Private/SlateEditorStyle.cpp
@@ -3959,7 +3959,18 @@ void FSlateEditorStyle::FStyle::SetupLevelEditorStyle()
 		Set( "LevelEditor.Recompile", new IMAGE_BRUSH( "Icons/icon_compile_40x", Icon40x40 ) );
 		Set( "LevelEditor.Recompile.Small", new IMAGE_BRUSH( "Icons/icon_compile_40x", Icon20x20 ) );
 
-		Set( "LevelEditor.ViewOptions", new IMAGE_BRUSH( "Icons/icon_view_40x", Icon40x40 ) );
+		Set("LevelEditor.SourceControl", new IMAGE_BRUSH("Icons/icon_source_control_40x", Icon40x40));
+		Set("LevelEditor.SourceControl.Small", new IMAGE_BRUSH("Icons/icon_source_control_40x", Icon20x20));
+		Set("LevelEditor.SourceControl.On", new IMAGE_BRUSH("Icons/icon_source_control_40x_on", Icon40x40));
+		Set("LevelEditor.SourceControl.On.Small", new IMAGE_BRUSH("Icons/icon_source_control_40x_on", Icon20x20));
+		Set("LevelEditor.SourceControl.Off", new IMAGE_BRUSH("Icons/icon_source_control_40x_off", Icon40x40));
+		Set("LevelEditor.SourceControl.Off.Small", new IMAGE_BRUSH("Icons/icon_source_control_40x_off", Icon20x20));
+		Set("LevelEditor.SourceControl.Unknown", new IMAGE_BRUSH("Icons/icon_source_control_40x_unknown", Icon40x40));
+		Set("LevelEditor.SourceControl.Unknown.Small", new IMAGE_BRUSH("Icons/icon_source_control_40x_unknown", Icon20x20));
+		Set("LevelEditor.SourceControl.Problem", new IMAGE_BRUSH("Icons/icon_source_control_40x_problem", Icon40x40));
+		Set("LevelEditor.SourceControl.Problem.Small", new IMAGE_BRUSH("Icons/icon_source_control_40x_problem", Icon20x20));
+
+		Set("LevelEditor.ViewOptions", new IMAGE_BRUSH("Icons/icon_view_40x", Icon40x40));
 		Set( "LevelEditor.ViewOptions.Small", new IMAGE_BRUSH( "Icons/icon_view_40x", Icon20x20 ) );
 
 		Set( "LevelEditor.GameSettings", new IMAGE_BRUSH( "Icons/icon_game_settings_40x", Icon40x40 ) );
diff --git a/Engine/Source/Editor/LevelEditor/LevelEditor.Build.cs b/Engine/Source/Editor/LevelEditor/LevelEditor.Build.cs
--- a/Engine/Source/Editor/LevelEditor/LevelEditor.Build.cs
+++ b/Engine/Source/Editor/LevelEditor/LevelEditor.Build.cs
@@ -40,6 +40,7 @@ public class LevelEditor : ModuleRules
 				"MessageLog",
 				"NewsFeed",
                 "SourceControl",
+                "SourceControlWindows",
                 "StatsViewer",
 				"UnrealEd", 
 				"RenderCore",
diff --git a/Engine/Source/Editor/LevelEditor/Private/LevelEditor.cpp b/Engine/Source/Editor/LevelEditor/Private/LevelEditor.cpp
--- a/Engine/Source/Editor/LevelEditor/Private/LevelEditor.cpp
+++ b/Engine/Source/Editor/LevelEditor/Private/LevelEditor.cpp
@@ -549,7 +549,29 @@ void FLevelEditorModule::BindGlobalLevelEditorCommands()
 	ActionList.MapAction( Commands.Build,
 		FExecuteAction::CreateStatic( &FLevelEditorActionCallbacks::Build_Execute ) );
 
-	ActionList.MapAction( Commands.RecompileGameCode,
+	ActionList.MapAction(
+		Commands.ConnectToSourceControl,
+		FExecuteAction::CreateStatic(&FLevelEditorActionCallbacks::ConnectToSourceControl_Clicked)
+		);
+
+	ActionList.MapAction(
+		Commands.ChangeSourceControlSettings,
+		FExecuteAction::CreateStatic(&FLevelEditorActionCallbacks::ConnectToSourceControl_Clicked)
+		);
+
+	ActionList.MapAction(
+		Commands.CheckOutModifiedFiles,
+		FExecuteAction::CreateStatic(&FLevelEditorActionCallbacks::CheckOutModifiedFiles_Clicked),
+		FCanExecuteAction::CreateStatic(&FLevelEditorActionCallbacks::CheckOutModifiedFiles_CanExecute)
+		);
+
+	ActionList.MapAction(
+		Commands.SubmitToSourceControl,
+		FExecuteAction::CreateStatic(&FLevelEditorActionCallbacks::SubmitToSourceControl_Clicked),
+		FCanExecuteAction::CreateStatic(&FLevelEditorActionCallbacks::SubmitToSourceControl_CanExecute)
+		);
+
+	ActionList.MapAction(Commands.RecompileGameCode,
 		FExecuteAction::CreateStatic( &FLevelEditorActionCallbacks::RecompileGameCode_Clicked ),
 		FCanExecuteAction::CreateStatic( &FLevelEditorActionCallbacks::Recompile_CanExecute )
 		);
diff --git a/Engine/Source/Editor/LevelEditor/Private/LevelEditorActions.cpp b/Engine/Source/Editor/LevelEditor/Private/LevelEditorActions.cpp
--- a/Engine/Source/Editor/LevelEditor/Private/LevelEditorActions.cpp
+++ b/Engine/Source/Editor/LevelEditor/Private/LevelEditorActions.cpp
@@ -49,6 +49,7 @@
 #include "EngineBuildSettings.h"
 #include "HotReloadInterface.h"
 #include "ISourceControlModule.h"
+#include "SourceControlWindows.h"
 #include "NotificationManager.h"
 #include "SNotificationList.h"
 #include "Engine/Selection.h"
@@ -921,6 +922,54 @@ bool FLevelEditorActionCallbacks::Recompile_CanExecute()
 	return !HotReloadSupport.IsCurrentlyCompiling() && !(GEngineVersion.IsPromotedBuild() && FEngineBuildSettings::IsPerforceBuild());
 }
 
+void FLevelEditorActionCallbacks::ConnectToSourceControl_Clicked()
+{
+	// Show login window regardless of current status - its useful as a shortcut to change settings.
+	ISourceControlModule& SourceControlModule = ISourceControlModule::Get();
+	SourceControlModule.ShowLoginDialog(FSourceControlLoginClosed(), ELoginWindowMode::Modeless, EOnLoginWindowStartup::PreserveProvider);
+}
+
+bool FLevelEditorActionCallbacks::CheckOutModifiedFiles_CanExecute()
+{
+	ISourceControlModule& SourceControlModule = ISourceControlModule::Get();
+	if (ISourceControlModule::Get().IsEnabled() &&
+		ISourceControlModule::Get().GetProvider().IsAvailable())
+	{
+		TArray<UPackage*> PackagesToSave;
+		FEditorFileUtils::GetDirtyWorldPackages(PackagesToSave);
+		FEditorFileUtils::GetDirtyContentPackages(PackagesToSave);
+
+		return PackagesToSave.Num() > 0;
+	}
+
+	return false;
+}
+
+void FLevelEditorActionCallbacks::CheckOutModifiedFiles_Clicked()
+{
+	TArray<UPackage*> PackagesToSave;
+	FEditorFileUtils::GetDirtyWorldPackages(PackagesToSave);
+	FEditorFileUtils::GetDirtyContentPackages(PackagesToSave);
+
+	const bool bCheckDirty = true;
+	const bool bPromptUserToSave = false;
+	FEditorFileUtils::PromptForCheckoutAndSave(PackagesToSave, bCheckDirty, bPromptUserToSave);
+}
+
+bool FLevelEditorActionCallbacks::SubmitToSourceControl_CanExecute()
+{
+	ISourceControlModule& SourceControlModule = ISourceControlModule::Get();
+	return ISourceControlModule::Get().IsEnabled() &&
+		   ISourceControlModule::Get().GetProvider().IsAvailable() &&
+		   FSourceControlWindows::CanChoosePackagesToCheckIn();
+}
+
+void FLevelEditorActionCallbacks::SubmitToSourceControl_Clicked()
+{
+	FSourceControlWindows::ChoosePackagesToCheckIn();
+}
+
+
 void FLevelEditorActionCallbacks::GoToCodeForActor_Clicked()
 {
 	const auto& SelectedActorInfo = AssetSelectionUtils::GetSelectedActorInfo();
@@ -2865,6 +2914,11 @@ void FLevelEditorCommands::RegisterCommands()
 	UI_COMMAND( MaterialQualityLevel_Low, "Low", "Sets material quality in the scene to low.", EUserInterfaceActionType::RadioButton, FInputGesture() );
 	UI_COMMAND( MaterialQualityLevel_High, "High", "Sets material quality in the scene to high.", EUserInterfaceActionType::RadioButton, FInputGesture() );
 
+	UI_COMMAND( ConnectToSourceControl, "Connect to Source Control...", "Opens a dialog to connect to source control.", EUserInterfaceActionType::Button, FInputGesture());
+	UI_COMMAND( ChangeSourceControlSettings, "Change Source Control Settings...", "Opens a dialog to change source control settings.", EUserInterfaceActionType::Button, FInputGesture());
+	UI_COMMAND( CheckOutModifiedFiles, "Check Out Modified Files...", "Opens a dialog to check out any assets which have been modified.", EUserInterfaceActionType::Button, FInputGesture());
+	UI_COMMAND( SubmitToSourceControl, "Submit to Source Control...", "Opens a dialog with check in options for content and levels.", EUserInterfaceActionType::Button, FInputGesture());
+
 	static const FText FeatureLevelLabels[ERHIFeatureLevel::Num] = 
 	{
 		NSLOCTEXT("LevelEditorCommands", "FeatureLevelPreviewType_ES2", "Mobile / HTML5"),
diff --git a/Engine/Source/Editor/LevelEditor/Private/LevelEditorToolBar.cpp b/Engine/Source/Editor/LevelEditor/Private/LevelEditorToolBar.cpp
--- a/Engine/Source/Editor/LevelEditor/Private/LevelEditorToolBar.cpp
+++ b/Engine/Source/Editor/LevelEditor/Private/LevelEditorToolBar.cpp
@@ -1179,6 +1179,99 @@ TSharedRef< SWidget > FLevelEditorToolBar::MakeLevelEditorToolBar( const TShared
 	}
 	ToolbarBuilder.EndSection();
 
+	ToolbarBuilder.BeginSection("SourceControl");
+	{
+		enum EQueryState
+		{
+			NotQueried,
+			Querying,
+			Queried,
+		};
+
+		static EQueryState QueryState = EQueryState::NotQueried;
+
+		struct FSourceControlStatus
+		{
+			static void CheckSourceControlStatus()
+			{
+				ISourceControlModule& SourceControlModule = ISourceControlModule::Get();
+				if (SourceControlModule.IsEnabled())
+				{
+					SourceControlModule.GetProvider().Execute(ISourceControlOperation::Create<FConnect>(),
+															  EConcurrency::Asynchronous,
+															  FSourceControlOperationComplete::CreateStatic(&FSourceControlStatus::OnSourceControlOperationComplete));
+					QueryState = EQueryState::Querying;
+				}
+			}
+
+			static void OnSourceControlOperationComplete(const FSourceControlOperationRef& InOperation, ECommandResult::Type InResult)
+			{
+				QueryState = EQueryState::Queried;
+			}
+
+			static FText GetSourceControlTooltip()
+			{
+				if (QueryState == EQueryState::NotQueried)
+				{
+					CheckSourceControlStatus();
+					return LOCTEXT("CheckingSourceControl", "Checking source control status...");
+				}
+				else if (QueryState == EQueryState::Querying)
+				{
+					return LOCTEXT("SourceControlUnknown", "Source control status is unknown");
+				}
+				else
+				{
+					return ISourceControlModule::Get().GetProvider().GetStatusText();
+				}
+			}
+
+			static FSlateIcon GetSourceControlIcon()
+			{
+				if (QueryState == EQueryState::Queried)
+				{
+					ISourceControlModule& SourceControlModule = ISourceControlModule::Get();
+					if (SourceControlModule.IsEnabled())
+					{
+						if (!SourceControlModule.GetProvider().IsAvailable())
+						{
+							return FSlateIcon(FEditorStyle::GetStyleSetName(), "LevelEditor.SourceControl.Error");
+						}
+						else
+						{
+							return FSlateIcon(FEditorStyle::GetStyleSetName(), "LevelEditor.SourceControl.On");
+						}
+					}
+					else
+					{
+						return FSlateIcon(FEditorStyle::GetStyleSetName(), "LevelEditor.SourceControl.Off");
+					}
+				}
+				else
+				{
+					if (QueryState == EQueryState::NotQueried)
+					{
+						CheckSourceControlStatus();
+					}
+
+					return FSlateIcon(FEditorStyle::GetStyleSetName(), "LevelEditor.SourceControl.Unknown");
+				}
+			}
+		};
+
+		FSourceControlStatus::CheckSourceControlStatus();
+
+		ToolbarBuilder.AddComboButton(
+			FUIAction(),
+			FOnGetContent::CreateStatic(&FLevelEditorToolBar::GenerateSourceControlMenu, InCommandList),
+			LOCTEXT("SourceControl_Label", "Source Control"),
+			TAttribute<FText>::Create(&FSourceControlStatus::GetSourceControlTooltip),
+			TAttribute<FSlateIcon>::Create(&FSourceControlStatus::GetSourceControlIcon),
+			false
+			);
+	}
+	ToolbarBuilder.EndSection();
+
 	// Create the tool bar!
 	return
 		SNew( SBorder )
@@ -1682,6 +1775,49 @@ TSharedRef< SWidget > FLevelEditorToolBar::GenerateQuickSettingsMenu( TSharedRef
 }
 
 
+TSharedRef< SWidget > FLevelEditorToolBar::GenerateSourceControlMenu(TSharedRef<FUICommandList> InCommandList)
+{
+#define LOCTEXT_NAMESPACE "LevelToolBarSourceControlMenu"
+
+	// Get all menu extenders for this context menu from the level editor module
+	FLevelEditorModule& LevelEditorModule = FModuleManager::GetModuleChecked<FLevelEditorModule>(TEXT("LevelEditor"));
+	TArray<FLevelEditorModule::FLevelEditorMenuExtender> MenuExtenderDelegates = LevelEditorModule.GetAllLevelEditorToolbarSourceControlMenuExtenders();
+
+	TArray<TSharedPtr<FExtender>> Extenders;
+	for (int32 i = 0; i < MenuExtenderDelegates.Num(); ++i)
+	{
+		if (MenuExtenderDelegates[i].IsBound())
+		{
+			Extenders.Add(MenuExtenderDelegates[i].Execute(InCommandList));
+		}
+	}
+	TSharedPtr<FExtender> MenuExtender = FExtender::Combine(Extenders);
+
+	const bool bShouldCloseWindowAfterMenuSelection = true;
+	FMenuBuilder MenuBuilder(bShouldCloseWindowAfterMenuSelection, InCommandList, MenuExtender);
+
+	MenuBuilder.BeginSection("SourceControlActions", LOCTEXT("SourceControlMenuHeadingActions", "Actions"));
+
+	ISourceControlModule& SourceControlModule = ISourceControlModule::Get();
+	if (ISourceControlModule::Get().IsEnabled() && ISourceControlModule::Get().GetProvider().IsAvailable())
+	{
+		MenuBuilder.AddMenuEntry(FLevelEditorCommands::Get().ChangeSourceControlSettings);
+	}
+	else
+	{
+		MenuBuilder.AddMenuEntry(FLevelEditorCommands::Get().ConnectToSourceControl);
+	}
+
+	MenuBuilder.AddMenuEntry(FLevelEditorCommands::Get().CheckOutModifiedFiles);
+	MenuBuilder.AddMenuEntry(FLevelEditorCommands::Get().SubmitToSourceControl);
+
+	MenuBuilder.EndSection();
+
+#undef LOCTEXT_NAMESPACE
+
+	return MenuBuilder.MakeWidget();
+}
+
 TSharedRef< SWidget > FLevelEditorToolBar::GenerateOpenBlueprintMenuContent( TSharedRef<FUICommandList> InCommandList, TWeakPtr< SLevelEditor > InLevelEditor )
 {
 #define LOCTEXT_NAMESPACE "LevelToolBarViewMenu"
diff --git a/Engine/Source/Editor/LevelEditor/Private/LevelEditorToolBar.h b/Engine/Source/Editor/LevelEditor/Private/LevelEditorToolBar.h
--- a/Engine/Source/Editor/LevelEditor/Private/LevelEditorToolBar.h
+++ b/Engine/Source/Editor/LevelEditor/Private/LevelEditorToolBar.h
@@ -49,6 +49,13 @@ protected:
 	 */
 	static TSharedRef< SWidget > GenerateQuickSettingsMenu( TSharedRef<FUICommandList> InCommandList );
 
+	/**
+	 * Generates menu content for the source control combo button drop down menu
+	 *
+	 * @return	Menu content widget
+	 */
+	static TSharedRef< SWidget > GenerateSourceControlMenu(TSharedRef<FUICommandList> InCommandList);
+
 	/**
 	 * Generates menu content for the compile combo button drop down menu
 	 *
diff --git a/Engine/Source/Editor/LevelEditor/Public/LevelEditor.h b/Engine/Source/Editor/LevelEditor/Public/LevelEditor.h
--- a/Engine/Source/Editor/LevelEditor/Public/LevelEditor.h
+++ b/Engine/Source/Editor/LevelEditor/Public/LevelEditor.h
@@ -205,7 +205,8 @@ public:
 	virtual TArray<FLevelEditorMenuExtender>& GetAllLevelEditorToolbarViewMenuExtenders() {return LevelEditorToolbarViewMenuExtenders;}
 	virtual TArray<FLevelEditorMenuExtender>& GetAllLevelEditorToolbarBuildMenuExtenders() {return LevelEditorToolbarBuildMenuExtenders;}
 	virtual TArray<FLevelEditorMenuExtender>& GetAllLevelEditorToolbarCompileMenuExtenders() {return LevelEditorToolbarCompileMenuExtenders;}
-	virtual TArray<FLevelEditorMenuExtender>& GetAllLevelEditorToolbarCreateMenuExtenders() {return LevelEditorToolbarCreateMenuExtenders;}
+	virtual TArray<FLevelEditorMenuExtender>& GetAllLevelEditorToolbarSourceControlMenuExtenders() { return LevelEditorToolbarSourceControlMenuExtenders; }
+	virtual TArray<FLevelEditorMenuExtender>& GetAllLevelEditorToolbarCreateMenuExtenders() { return LevelEditorToolbarCreateMenuExtenders; }
 	
 	/** Gets the extensibility managers for outside entities to extend static mesh editor's menus and toolbars */
 	virtual TSharedPtr<FExtensibilityManager> GetMenuExtensibilityManager() {return MenuExtensibilityManager;}
@@ -276,6 +277,7 @@ private:
 	TArray<FLevelEditorMenuExtender> LevelEditorToolbarViewMenuExtenders;
 	TArray<FLevelEditorMenuExtender> LevelEditorToolbarBuildMenuExtenders;
 	TArray<FLevelEditorMenuExtender> LevelEditorToolbarCompileMenuExtenders;
+	TArray<FLevelEditorMenuExtender> LevelEditorToolbarSourceControlMenuExtenders;
 	TArray<FLevelEditorMenuExtender> LevelEditorToolbarCreateMenuExtenders;
 
 	/* Pointer to the current level Editor instance */
diff --git a/Engine/Source/Editor/LevelEditor/Public/LevelEditorActions.h b/Engine/Source/Editor/LevelEditor/Public/LevelEditorActions.h
--- a/Engine/Source/Editor/LevelEditor/Public/LevelEditorActions.h
+++ b/Engine/Source/Editor/LevelEditor/Public/LevelEditorActions.h
@@ -550,6 +550,14 @@ public:
 	//TSharedPtr< FUICommandInfo > LandscapeMode;
 	//TSharedPtr< FUICommandInfo > FoliageMode;
 
+	/**
+	 * Source Control Commands
+	 */
+	TSharedPtr< FUICommandInfo > ConnectToSourceControl;
+	TSharedPtr< FUICommandInfo > ChangeSourceControlSettings;
+	TSharedPtr< FUICommandInfo > CheckOutModifiedFiles;
+	TSharedPtr< FUICommandInfo > SubmitToSourceControl;
+
 	/**
 	 * Misc Commands
 	 */
@@ -751,6 +759,23 @@ public:
 	static void RecompileGameCode_Clicked();
 	static bool Recompile_CanExecute();
 
+	/**
+	 * Called when requesting connection to source control
+	 */
+	static void ConnectToSourceControl_Clicked();
+
+	/**
+	 * Called when Check Out Modified Files is clicked
+	 */
+	static void CheckOutModifiedFiles_Clicked();
+	static bool CheckOutModifiedFiles_CanExecute();
+
+	/**
+	 * Called when Submit to Source Control is clicked
+	 */
+	static void SubmitToSourceControl_Clicked();
+	static bool SubmitToSourceControl_CanExecute();
+
 	/**
 	 * Called when the FindInContentBrowser command is executed
 	 */
diff --git a/Engine/Source/Editor/MainFrame/Private/Frame/MainFrameActions.cpp b/Engine/Source/Editor/MainFrame/Private/Frame/MainFrameActions.cpp
--- a/Engine/Source/Editor/MainFrame/Private/Frame/MainFrameActions.cpp
+++ b/Engine/Source/Editor/MainFrame/Private/Frame/MainFrameActions.cpp
@@ -183,181 +183,14 @@ void FMainFrameActionCallbacks::ChoosePackagesToSave()
 	FEditorFileUtils::SaveDirtyPackages( bPromptUserToSave, bSaveMapPackages, bSaveContentPackages, bFastSave, bNotifyNoPackagesSaved );
 }
 
-void FMainFrameActionCallbacks::ChoosePackagesToCheckInCompleted(const TArray<UPackage*>& LoadedPackages, const TArray<FString>& PackageNames, const TArray<FString>& ConfigFiles)
-{
-	if (ChoosePackagesToCheckInNotification.IsValid())
-	{
-		ChoosePackagesToCheckInNotification.Pin()->ExpireAndFadeout();
-	}
-	ChoosePackagesToCheckInNotification.Reset();
-	
-	check(PackageNames.Num() > 0 || ConfigFiles.Num() > 0);
-
-	// Prompt the user to ask if they would like to first save any dirty packages they are trying to check-in
-	const FEditorFileUtils::EPromptReturnCode UserResponse = FEditorFileUtils::PromptForCheckoutAndSave( LoadedPackages, true, true );
-
-	// If the user elected to save dirty packages, but one or more of the packages failed to save properly OR if the user
-	// canceled out of the prompt, don't follow through on the check-in process
-	const bool bShouldProceed = ( UserResponse == FEditorFileUtils::EPromptReturnCode::PR_Success || UserResponse == FEditorFileUtils::EPromptReturnCode::PR_Declined );
-	if ( bShouldProceed )
-	{
-		FAssetToolsModule& AssetToolsModule = FModuleManager::LoadModuleChecked<FAssetToolsModule>(TEXT("AssetTools"));
-		FSourceControlWindows::PromptForCheckin(PackageNames, ConfigFiles);
-	}
-	else
-	{
-		// If a failure occurred, alert the user that the check-in was aborted. This warning shouldn't be necessary if the user cancelled
-		// from the dialog, because they obviously intended to cancel the whole operation.
-		if ( UserResponse == FEditorFileUtils::EPromptReturnCode::PR_Failure )
-		{
-			FMessageDialog::Open( EAppMsgType::Ok, NSLOCTEXT("UnrealEd", "SCC_Checkin_Aborted", "Check-in aborted as a result of save failure.") );
-		}
-	}
-}
-
-void FMainFrameActionCallbacks::ChoosePackagesToCheckInCancelled(FSourceControlOperationRef InOperation)
-{
-	ISourceControlProvider& SourceControlProvider = ISourceControlModule::Get().GetProvider();
-	SourceControlProvider.CancelOperation(InOperation);
-
-	if (ChoosePackagesToCheckInNotification.IsValid())
-	{
-		ChoosePackagesToCheckInNotification.Pin()->ExpireAndFadeout();
-	}
-	ChoosePackagesToCheckInNotification.Reset();
-}
-
-void FMainFrameActionCallbacks::ChoosePackagesToCheckInCallback(const FSourceControlOperationRef& InOperation, ECommandResult::Type InResult)
-{
-	if (ChoosePackagesToCheckInNotification.IsValid())
-	{
-		ChoosePackagesToCheckInNotification.Pin()->ExpireAndFadeout();
-	}
-	ChoosePackagesToCheckInNotification.Reset();
-
-	if(InResult == ECommandResult::Succeeded)
-	{
-		// Get a list of all the checked out packages
-		int32 NumSelectedNames = 0;
-		int32 NumSelectedPackages = 0;
-		TArray<FString> PackageNames;
-		TArray<UPackage*> LoadedPackages;
-		TMap<FString, FSourceControlStatePtr> PackageStates;
-		FEditorFileUtils::FindAllSubmittablePackageFiles( PackageStates, true );
-		for (TMap<FString, FSourceControlStatePtr>::TConstIterator PackageIter(PackageStates); PackageIter; ++PackageIter)
-		{
-			const FString Filename = *PackageIter.Key();
-			const FString PackageName = FPackageName::FilenameToLongPackageName(Filename);
-			const FSourceControlStatePtr CurPackageSCCState = PackageIter.Value();
-				
-			UPackage* Package = FindPackage(NULL, *PackageName);
-
-			// Put pre-selected items at the start of the list
-			if (CurPackageSCCState.IsValid() && CurPackageSCCState->IsSourceControlled())
-			{
-				if (Package != NULL)
-				{
-					LoadedPackages.Insert(Package, NumSelectedPackages++);
-				}
-				PackageNames.Insert(PackageName, NumSelectedNames++);
-			}
-			else
-			{
-				if (Package != NULL)
-				{
-					LoadedPackages.Add(Package);
-				}
-				PackageNames.Add(PackageName);
-			}
-		}
-
-		// Get a list of all the checked out config files
-		TMap<FString, FSourceControlStatePtr> ConfigFileStates;
-		TArray<FString> ConfigFilesToSubmit;
-		FEditorFileUtils::FindAllSubmittableConfigFiles(ConfigFileStates);
-		for (TMap<FString, FSourceControlStatePtr>::TConstIterator It(ConfigFileStates); It; ++It)
-		{
-			ConfigFilesToSubmit.Add(It.Key());
-		}
-
-		if ( PackageNames.Num() > 0 || ConfigFilesToSubmit.Num() > 0 )
-		{
-			ChoosePackagesToCheckInCompleted(LoadedPackages, PackageNames, ConfigFilesToSubmit);
-		}
-		else
-		{
-			FMessageLog EditorErrors("EditorErrors");
-			EditorErrors.Warning(LOCTEXT( "NoAssetsToCheckIn", "No assets to check in!"));
-			EditorErrors.Notify();
-		}
-	}
-	else if(InResult == ECommandResult::Failed)
-	{
-		FMessageLog EditorErrors("EditorErrors");
-		EditorErrors.Warning(LOCTEXT( "CheckInOperationFailed", "Failed checking source control status!"));
-		EditorErrors.Notify();
-	}
-}
-
 void FMainFrameActionCallbacks::ChoosePackagesToCheckIn()
 {
-	if ( ISourceControlModule::Get().IsEnabled() )
-	{
-		if( ISourceControlModule::Get().GetProvider().IsAvailable() )
-		{
-			// make sure we update the SCC status of all packages (this could take a long time, so we will run it as a background task)
-			TArray<FString> Packages;
-			FEditorFileUtils::FindAllPackageFiles(Packages);
-
-			// Get list of filenames corresponding to packages
-			TArray<FString> Filenames = SourceControlHelpers::PackageFilenames(Packages);
-
-			// Add game config files to the list
-			FEditorFileUtils::FindAllConfigFiles(Filenames);
-			
-			ISourceControlProvider& SourceControlProvider = ISourceControlModule::Get().GetProvider();
-			FSourceControlOperationRef Operation = ISourceControlOperation::Create<FUpdateStatus>();
-			SourceControlProvider.Execute(Operation, Filenames, EConcurrency::Asynchronous, FSourceControlOperationComplete::CreateStatic(&FMainFrameActionCallbacks::ChoosePackagesToCheckInCallback));
-
-			if (ChoosePackagesToCheckInNotification.IsValid())
-			{
-				ChoosePackagesToCheckInNotification.Pin()->ExpireAndFadeout();
-			}
-
-			FNotificationInfo Info(LOCTEXT("ChooseAssetsToCheckInIndicator", "Checking for assets to check in..."));
-			Info.bFireAndForget = false;
-			Info.ExpireDuration = 0.0f;
-			Info.FadeOutDuration = 1.0f;
-
-			if(SourceControlProvider.CanCancelOperation(Operation))
-			{
-				Info.ButtonDetails.Add(FNotificationButtonInfo(
-					LOCTEXT("ChoosePackagesToCheckIn_CancelButton", "Cancel"), 
-					LOCTEXT("ChoosePackagesToCheckIn_CancelButtonTooltip", "Cancel the check in operation."), 
-					FSimpleDelegate::CreateStatic(&FMainFrameActionCallbacks::ChoosePackagesToCheckInCancelled, Operation)
-					));
-			}
-
-			ChoosePackagesToCheckInNotification = FSlateNotificationManager::Get().AddNotification(Info);
-
-			if (ChoosePackagesToCheckInNotification.IsValid())
-			{
-				ChoosePackagesToCheckInNotification.Pin()->SetCompletionState(SNotificationItem::CS_Pending);
-			}
-		}
-		else
-		{
-			FMessageLog EditorErrors("EditorErrors");
-			EditorErrors.Warning(LOCTEXT( "NoSCCConnection", "No connection to source control available!"))
-				->AddToken(FDocumentationToken::Create(TEXT("Engine/UI/SourceControl")));
-			EditorErrors.Notify();
-		}
-	}
+	FSourceControlWindows::ChoosePackagesToCheckIn();
 }
 
 bool FMainFrameActionCallbacks::CanChoosePackagesToCheckIn()
 {
-	return !ChoosePackagesToCheckInNotification.IsValid();
+	return FSourceControlWindows::CanChoosePackagesToCheckIn();
 }
 
 void FMainFrameActionCallbacks::ConnectToSourceControl()
diff --git a/Engine/Source/Editor/MainFrame/Private/Frame/MainFrameActions.h b/Engine/Source/Editor/MainFrame/Private/Frame/MainFrameActions.h
--- a/Engine/Source/Editor/MainFrame/Private/Frame/MainFrameActions.h
+++ b/Engine/Source/Editor/MainFrame/Private/Frame/MainFrameActions.h
@@ -262,15 +262,6 @@ protected:
 	 */
 	static FString TranslateUATError(int32 ErrorCode);
 
-	/** Callback for ChoosePackagesToCheckIn(), continues to bring up UI once source control operations are complete */
-	static void ChoosePackagesToCheckInCallback(const FSourceControlOperationRef& InOperation, ECommandResult::Type InResult);
-
-	/** Called when the process has completed and we have packages to check in */
-	static void ChoosePackagesToCheckInCompleted(const TArray<UPackage*>& LoadedPackages, const TArray<FString>& PackageNames, const TArray<FString>& ConfigFiles);
-
-	/** Delegate called when the user has decided to cancel the check in process */
-	static void ChoosePackagesToCheckInCancelled(FSourceControlOperationRef InOperation);
-
 private:
 
 	struct EventData
diff --git a/Engine/Source/Editor/SourceControlWindows/Private/SSourceControlSubmit.cpp b/Engine/Source/Editor/SourceControlWindows/Private/SSourceControlSubmit.cpp
--- a/Engine/Source/Editor/SourceControlWindows/Private/SSourceControlSubmit.cpp
+++ b/Engine/Source/Editor/SourceControlWindows/Private/SSourceControlSubmit.cpp
@@ -5,6 +5,7 @@
 #include "MessageLog.h"
 #include "SNotificationList.h"
 #include "NotificationManager.h"
+#include "FileHelpers.h"
 
 
 IMPLEMENT_MODULE( FDefaultModuleImpl, SourceControlWindows );
@@ -451,6 +452,186 @@ private:
 	ECheckBoxState	KeepCheckedOut;
 };
 
+
+TWeakPtr<SNotificationItem> FSourceControlWindows::ChoosePackagesToCheckInNotification;
+
+void FSourceControlWindows::ChoosePackagesToCheckInCompleted(const TArray<UPackage*>& LoadedPackages, const TArray<FString>& PackageNames, const TArray<FString>& ConfigFiles)
+{
+	if (ChoosePackagesToCheckInNotification.IsValid())
+	{
+		ChoosePackagesToCheckInNotification.Pin()->ExpireAndFadeout();
+	}
+	ChoosePackagesToCheckInNotification.Reset();
+
+	check(PackageNames.Num() > 0 || ConfigFiles.Num() > 0);
+
+	// Prompt the user to ask if they would like to first save any dirty packages they are trying to check-in
+	const FEditorFileUtils::EPromptReturnCode UserResponse = FEditorFileUtils::PromptForCheckoutAndSave(LoadedPackages, true, true);
+
+	// If the user elected to save dirty packages, but one or more of the packages failed to save properly OR if the user
+	// canceled out of the prompt, don't follow through on the check-in process
+	const bool bShouldProceed = (UserResponse == FEditorFileUtils::EPromptReturnCode::PR_Success || UserResponse == FEditorFileUtils::EPromptReturnCode::PR_Declined);
+	if (bShouldProceed)
+	{
+		FAssetToolsModule& AssetToolsModule = FModuleManager::LoadModuleChecked<FAssetToolsModule>(TEXT("AssetTools"));
+		FSourceControlWindows::PromptForCheckin(PackageNames, ConfigFiles);
+	}
+	else
+	{
+		// If a failure occurred, alert the user that the check-in was aborted. This warning shouldn't be necessary if the user cancelled
+		// from the dialog, because they obviously intended to cancel the whole operation.
+		if (UserResponse == FEditorFileUtils::EPromptReturnCode::PR_Failure)
+		{
+			FMessageDialog::Open(EAppMsgType::Ok, NSLOCTEXT("UnrealEd", "SCC_Checkin_Aborted", "Check-in aborted as a result of save failure."));
+		}
+	}
+}
+
+void FSourceControlWindows::ChoosePackagesToCheckInCancelled(FSourceControlOperationRef InOperation)
+{
+	ISourceControlProvider& SourceControlProvider = ISourceControlModule::Get().GetProvider();
+	SourceControlProvider.CancelOperation(InOperation);
+
+	if (ChoosePackagesToCheckInNotification.IsValid())
+	{
+		ChoosePackagesToCheckInNotification.Pin()->ExpireAndFadeout();
+	}
+	ChoosePackagesToCheckInNotification.Reset();
+}
+
+void FSourceControlWindows::ChoosePackagesToCheckInCallback(const FSourceControlOperationRef& InOperation, ECommandResult::Type InResult)
+{
+	if (ChoosePackagesToCheckInNotification.IsValid())
+	{
+		ChoosePackagesToCheckInNotification.Pin()->ExpireAndFadeout();
+	}
+	ChoosePackagesToCheckInNotification.Reset();
+
+	if (InResult == ECommandResult::Succeeded)
+	{
+		// Get a list of all the checked out packages
+		int32 NumSelectedNames = 0;
+		int32 NumSelectedPackages = 0;
+		TArray<FString> PackageNames;
+		TArray<UPackage*> LoadedPackages;
+		TMap<FString, FSourceControlStatePtr> PackageStates;
+		FEditorFileUtils::FindAllSubmittablePackageFiles(PackageStates, true);
+		for (TMap<FString, FSourceControlStatePtr>::TConstIterator PackageIter(PackageStates); PackageIter; ++PackageIter)
+		{
+			const FString Filename = *PackageIter.Key();
+			const FString PackageName = FPackageName::FilenameToLongPackageName(Filename);
+			const FSourceControlStatePtr CurPackageSCCState = PackageIter.Value();
+
+			UPackage* Package = FindPackage(NULL, *PackageName);
+
+			// Put pre-selected items at the start of the list
+			if (CurPackageSCCState.IsValid() && CurPackageSCCState->IsSourceControlled())
+			{
+				if (Package != NULL)
+				{
+					LoadedPackages.Insert(Package, NumSelectedPackages++);
+				}
+				PackageNames.Insert(PackageName, NumSelectedNames++);
+			}
+			else
+			{
+				if (Package != NULL)
+				{
+					LoadedPackages.Add(Package);
+				}
+				PackageNames.Add(PackageName);
+			}
+		}
+
+		// Get a list of all the checked out config files
+		TMap<FString, FSourceControlStatePtr> ConfigFileStates;
+		TArray<FString> ConfigFilesToSubmit;
+		FEditorFileUtils::FindAllSubmittableConfigFiles(ConfigFileStates);
+		for (TMap<FString, FSourceControlStatePtr>::TConstIterator It(ConfigFileStates); It; ++It)
+		{
+			ConfigFilesToSubmit.Add(It.Key());
+		}
+
+		if (PackageNames.Num() > 0 || ConfigFilesToSubmit.Num() > 0)
+		{
+			ChoosePackagesToCheckInCompleted(LoadedPackages, PackageNames, ConfigFilesToSubmit);
+		}
+		else
+		{
+			FMessageLog EditorErrors("EditorErrors");
+			EditorErrors.Warning(LOCTEXT("NoAssetsToCheckIn", "No assets to check in!"));
+			EditorErrors.Notify();
+		}
+	}
+	else if (InResult == ECommandResult::Failed)
+	{
+		FMessageLog EditorErrors("EditorErrors");
+		EditorErrors.Warning(LOCTEXT("CheckInOperationFailed", "Failed checking source control status!"));
+		EditorErrors.Notify();
+	}
+}
+
+void FSourceControlWindows::ChoosePackagesToCheckIn()
+{
+	if (ISourceControlModule::Get().IsEnabled())
+	{
+		if (ISourceControlModule::Get().GetProvider().IsAvailable())
+		{
+			// make sure we update the SCC status of all packages (this could take a long time, so we will run it as a background task)
+			TArray<FString> Packages;
+			FEditorFileUtils::FindAllPackageFiles(Packages);
+
+			// Get list of filenames corresponding to packages
+			TArray<FString> Filenames = SourceControlHelpers::PackageFilenames(Packages);
+
+			// Add game config files to the list
+			FEditorFileUtils::FindAllConfigFiles(Filenames);
+
+			ISourceControlProvider& SourceControlProvider = ISourceControlModule::Get().GetProvider();
+			FSourceControlOperationRef Operation = ISourceControlOperation::Create<FUpdateStatus>();
+			SourceControlProvider.Execute(Operation, Filenames, EConcurrency::Asynchronous, FSourceControlOperationComplete::CreateStatic(&FSourceControlWindows::ChoosePackagesToCheckInCallback));
+
+			if (ChoosePackagesToCheckInNotification.IsValid())
+			{
+				ChoosePackagesToCheckInNotification.Pin()->ExpireAndFadeout();
+			}
+
+			FNotificationInfo Info(LOCTEXT("ChooseAssetsToCheckInIndicator", "Checking for assets to check in..."));
+			Info.bFireAndForget = false;
+			Info.ExpireDuration = 0.0f;
+			Info.FadeOutDuration = 1.0f;
+
+			if (SourceControlProvider.CanCancelOperation(Operation))
+			{
+				Info.ButtonDetails.Add(FNotificationButtonInfo(
+					LOCTEXT("ChoosePackagesToCheckIn_CancelButton", "Cancel"),
+					LOCTEXT("ChoosePackagesToCheckIn_CancelButtonTooltip", "Cancel the check in operation."),
+					FSimpleDelegate::CreateStatic(&FSourceControlWindows::ChoosePackagesToCheckInCancelled, Operation)
+					));
+			}
+
+			ChoosePackagesToCheckInNotification = FSlateNotificationManager::Get().AddNotification(Info);
+
+			if (ChoosePackagesToCheckInNotification.IsValid())
+			{
+				ChoosePackagesToCheckInNotification.Pin()->SetCompletionState(SNotificationItem::CS_Pending);
+			}
+		}
+		else
+		{
+			FMessageLog EditorErrors("EditorErrors");
+			EditorErrors.Warning(LOCTEXT("NoSCCConnection", "No connection to source control available!"))
+				->AddToken(FDocumentationToken::Create(TEXT("Engine/UI/SourceControl")));
+			EditorErrors.Notify();
+		}
+	}
+}
+
+bool FSourceControlWindows::CanChoosePackagesToCheckIn()
+{
+	return !ChoosePackagesToCheckInNotification.IsValid();
+}
+
 static void FindFilesForCheckIn(const TArray<FString>& InFilenames, TArray<FString>& OutAddFiles, TArray<FString>& OutOpenFiles)
 {
 	ISourceControlProvider& SourceControlProvider = ISourceControlModule::Get().GetProvider();
@@ -461,7 +642,7 @@ static void FindFilesForCheckIn(const TArray<FString>& InFilenames, TArray<FStri
 	for (const FString& Filename : InFilenames)
 	{
 		FSourceControlStatePtr SourceControlState = SourceControlProvider.GetState(Filename, EStateCacheUsage::Use);
-		if(SourceControlState.IsValid())
+		if (SourceControlState.IsValid())
 		{
 			if (SourceControlState->CanCheckIn())
 			{
@@ -469,7 +650,7 @@ static void FindFilesForCheckIn(const TArray<FString>& InFilenames, TArray<FStri
 			}
 			else
 			{
-				if( !SourceControlState->IsSourceControlled() )
+				if (!SourceControlState->IsSourceControlled())
 				{
 					OutAddFiles.Add(Filename);
 				}
@@ -478,7 +659,6 @@ static void FindFilesForCheckIn(const TArray<FString>& InFilenames, TArray<FStri
 	}
 }
 
-
 bool FSourceControlWindows::PromptForCheckin(const TArray<FString>& InPackageNames, const TArray<FString>& InConfigFiles)
 {
 	bool bCheckInSuccess = true;
@@ -588,6 +768,7 @@ bool FSourceControlWindows::PromptForCheckin(const TArray<FString>& InPackageNam
 	return bCheckInSuccess;
 }
 
+
 #undef LOCTEXT_NAMESPACE
 
 #endif // SOURCE_CONTROL_WITH_SLATE
diff --git a/Engine/Source/Editor/SourceControlWindows/Public/SourceControlWindows.h b/Engine/Source/Editor/SourceControlWindows/Public/SourceControlWindows.h
--- a/Engine/Source/Editor/SourceControlWindows/Public/SourceControlWindows.h
+++ b/Engine/Source/Editor/SourceControlWindows/Public/SourceControlWindows.h
@@ -7,6 +7,12 @@
 class SOURCECONTROLWINDOWS_API FSourceControlWindows
 {
 public:
+	/** Opens a dialog to choose packages to submit */
+	static void ChoosePackagesToCheckIn();
+
+	/** Determines whether we can choose packages to check in (we cant if an operation is already in progress) */
+	static bool CanChoosePackagesToCheckIn();
+
 	/**
 	 * Display check in dialog for the specified packages
 	 *
@@ -30,5 +36,19 @@ public:
 	 * @return	true if the files were reverted; false if the user canceled out of the dialog
 	 */
 	static bool PromptForRevert(const TArray<FString>& InPackageNames );
+
+protected:
+	/** Callback for ChoosePackagesToCheckIn(), continues to bring up UI once source control operations are complete */
+	static void ChoosePackagesToCheckInCallback(const FSourceControlOperationRef& InOperation, ECommandResult::Type InResult);
+
+	/** Called when the process has completed and we have packages to check in */
+	static void ChoosePackagesToCheckInCompleted(const TArray<UPackage*>& LoadedPackages, const TArray<FString>& PackageNames, const TArray<FString>& ConfigFiles);
+
+	/** Delegate called when the user has decided to cancel the check in process */
+	static void ChoosePackagesToCheckInCancelled(FSourceControlOperationRef InOperation);
+
+private:
+	/** The notification in place while we choose packages to check in */
+	static TWeakPtr<SNotificationItem> ChoosePackagesToCheckInNotification;
 };
 
