Fix Mac build error caused by CL 2405840.

diff --git a/Engine/Source/Editor/SourceControlWindows/Public/SourceControlWindows.h b/Engine/Source/Editor/SourceControlWindows/Public/SourceControlWindows.h
--- a/Engine/Source/Editor/SourceControlWindows/Public/SourceControlWindows.h
+++ b/Engine/Source/Editor/SourceControlWindows/Public/SourceControlWindows.h
@@ -3,6 +3,7 @@
 #pragma once
 
 #include "Core.h"
+#include "ISourceControlProvider.h"
 
 class SOURCECONTROLWINDOWS_API FSourceControlWindows
 {
@@ -49,6 +50,6 @@ protected:
 
 private:
 	/** The notification in place while we choose packages to check in */
-	static TWeakPtr<SNotificationItem> ChoosePackagesToCheckInNotification;
+	static TWeakPtr<class SNotificationItem> ChoosePackagesToCheckInNotification;
 };
 
