Confirmed that source control revert works on unreferenced assets, they get reloaded. Added warning for when the asset is in use and referenced.

UE-16722

diff --git a/Engine/Source/Editor/SourceControlWindows/Private/SSourceControlRevert.cpp b/Engine/Source/Editor/SourceControlWindows/Private/SSourceControlRevert.cpp
--- a/Engine/Source/Editor/SourceControlWindows/Private/SSourceControlRevert.cpp
+++ b/Engine/Source/Editor/SourceControlWindows/Private/SSourceControlRevert.cpp
@@ -4,6 +4,9 @@
 #include "AssetToolsModule.h"
 #include "PackageTools.h"
 
+#include "SNotificationList.h"
+#include "NotificationManager.h"
+
 #define LOCTEXT_NAMESPACE "SSourceControlRevert"
 
 //-------------------------------------
@@ -429,6 +432,13 @@ bool FSourceControlWindows::PromptForRevert( const TArray<FString>& InPackageNam
 						if (Package != NULL)
 						{
 							Package->ClearFlags(RF_WasLoaded);
+							
+							// Create and display a notification about the revert failing (the object was still being used)
+							const FText NotificationErrorText = FText::Format(LOCTEXT("RevertFailed", "Failed to revert file {0} since it is currently in use."), 
+								FText::AsCultureInvariant(Package->GetName()));
+							FNotificationInfo Info(NotificationErrorText);
+							Info.ExpireDuration = 2.0f;
+							FSlateNotificationManager::Get().AddNotification(Info);
 						}
 					}
 
