PR #4667: Source control history: remove empty spacing in the right of the detail panel (Contributed by [REDACTED_PERSON])

diff --git a/Engine/Source/Editor/SourceControlWindows/Private/SSourceControlHistory.cpp b/Engine/Source/Editor/SourceControlWindows/Private/SSourceControlHistory.cpp
--- a/Engine/Source/Editor/SourceControlWindows/Private/SSourceControlHistory.cpp
+++ b/Engine/Source/Editor/SourceControlWindows/Private/SSourceControlHistory.cpp
@@ -810,13 +810,6 @@ private:
 					+SVerticalBox::Slot()
 					.FillHeight(0.25f)
 					.Padding(Padding)
-					[
-						//empty for spacing
-						SNullWidget::NullWidget
-					]
-					+SVerticalBox::Slot()
-					.FillHeight(0.25f)
-					.Padding(Padding)
 					[
 						SNew(STextBlock)
 						.Text(NSLOCTEXT("SourceControl.HistoryPanel.Info", "Changelist", "Changelist:"))
@@ -852,13 +845,6 @@ private:
 					+SVerticalBox::Slot()
 					.FillHeight(0.25f)
 					.Padding(Padding)
-					[
-						//empty for spacing
-						SNullWidget::NullWidget
-					]
-					+SVerticalBox::Slot()
-					.FillHeight(0.25f)
-					.Padding(Padding)
 					[
 						SNew(STextBlock)
 						.Text(this, &SSourceControlHistoryWidget::GetChangelistNumber)	
