PR #5120: Source Control: Add a "Diff Against Depot" context action to the Submit window (Contributed by [REDACTED_PERSON])

diff --git a/Engine/Source/Editor/SourceControlWindows/Private/SSourceControlSubmit.cpp b/Engine/Source/Editor/SourceControlWindows/Private/SSourceControlSubmit.cpp
--- a/Engine/Source/Editor/SourceControlWindows/Private/SSourceControlSubmit.cpp
+++ b/Engine/Source/Editor/SourceControlWindows/Private/SSourceControlSubmit.cpp
@@ -7,6 +7,7 @@
 #include "ISourceControlModule.h"
 #include "SourceControlHelpers.h"
 #include "Modules/ModuleManager.h"
+#include "Framework/MultiBox/MultiBoxBuilder.h"
 #include "Widgets/SWindow.h"
 #include "Widgets/Layout/SSpacer.h"
 #include "Widgets/Layout/SWrapBox.h"
@@ -18,6 +19,8 @@
 #include "Widgets/Notifications/SErrorText.h"
 #include "Widgets/Input/SCheckBox.h"
 #include "EditorStyleSet.h"
+#include "AssetToolsModule.h"
+#include "AssetRegistryModule.h"
 
 
 #if SOURCE_CONTROL_WITH_SLATE
@@ -142,8 +145,10 @@ void SSourceControlSubmitWidget::Construct(const FArguments& InArgs)
 					.ItemHeight(20)
 					.ListItemsSource(&ListViewItems)
 					.OnGenerateRow(this, &SSourceControlSubmitWidget::OnGenerateRowForList)
+					.OnContextMenuOpening(this, &SSourceControlSubmitWidget::OnCreateContextMenu)
+					.OnMouseButtonDoubleClick(this, &SSourceControlSubmitWidget::OnDiffAgainstDepotSelected)
 					.HeaderRow(HeaderRowWidget)
-					.SelectionMode(ESelectionMode::None)
+					.SelectionMode(ESelectionMode::Single)
 				]
 			]
 			+SVerticalBox::Slot()
@@ -216,6 +221,77 @@ void SSourceControlSubmitWidget::Construct(const FArguments& InArgs)
 	ParentFrame.Pin()->SetWidgetToFocusOnActivate(ChangeListDescriptionTextCtrl);
 }
 
+/** Corvus: Called to create a context menu when right-clicking on an item */
+TSharedPtr<SWidget> SSourceControlSubmitWidget::OnCreateContextMenu()
+{
+	if (SSourceControlSubmitWidget::CanDiffAgainstDepot())
+	{
+		FMenuBuilder MenuBuilder(true, NULL);
+
+		MenuBuilder.BeginSection("Source Control", NSLOCTEXT("SourceControl.SubmitWindow.Menu", "SourceControlSectionHeader", "Source Control"));
+		{
+			MenuBuilder.AddMenuEntry(
+				NSLOCTEXT("SourceControl.SubmitWindow.Menu", "DiffAgainstDepot", "Diff Against Depot"),
+				NSLOCTEXT("SourceControl.SubmitWindow.Menu", "DiffAgainstDepotTooltip", "Look at differences between your version of the asset and that in source control."),
+				FSlateIcon(FEditorStyle::GetStyleSetName(), "SourceControl.Actions.Diff"),
+				FUIAction(
+					FExecuteAction::CreateSP(this, &SSourceControlSubmitWidget::OnDiffAgainstDepot),
+					FCanExecuteAction::CreateSP(this, &SSourceControlSubmitWidget::CanDiffAgainstDepot)
+				)
+			);
+		}
+		MenuBuilder.EndSection();
+
+		return MenuBuilder.MakeWidget();
+	}
+	else
+	{
+		return nullptr;
+	}
+}
+
+bool SSourceControlSubmitWidget::CanDiffAgainstDepot() const
+{
+	bool bCanDiff = false;
+	const auto& SelectedItems = ListView->GetSelectedItems();
+	if (SelectedItems.Num() == 1)
+	{
+		bCanDiff = SelectedItems[0]->CanDiff();
+	}
+	return bCanDiff;
+}
+
+void SSourceControlSubmitWidget::OnDiffAgainstDepot()
+{
+	const auto& SelectedItems = ListView->GetSelectedItems();
+	if (SelectedItems.Num() == 1)
+	{
+		OnDiffAgainstDepotSelected(SelectedItems[0]);
+	}
+}
+
+void SSourceControlSubmitWidget::OnDiffAgainstDepotSelected(TSharedPtr<FSubmitItem> InSelectedItem)
+{
+	FString PackageName;
+	if (FPackageName::TryConvertFilenameToLongPackageName(InSelectedItem->GetFilename(), PackageName))
+	{
+		TArray<FAssetData> Assets;
+		FAssetRegistryModule& AssetRegistryModule = FModuleManager::LoadModuleChecked<FAssetRegistryModule>(TEXT("AssetRegistry"));
+		AssetRegistryModule.Get().GetAssetsByPackageName(*PackageName, Assets);
+		if (Assets.Num() == 1)
+		{
+			const FAssetData& AssetData = Assets[0];
+			UObject* CurrentObject = AssetData.GetAsset();
+			if (CurrentObject)
+			{
+				const FString AssetName = AssetData.AssetName.ToString();
+				FAssetToolsModule& AssetToolsModule = FModuleManager::GetModuleChecked<FAssetToolsModule>("AssetTools");
+				AssetToolsModule.Get().DiffAgainstDepot(CurrentObject, PackageName, AssetName);
+			}
+		}
+	}
+}
+
 FReply SSourceControlSubmitWidget::OnKeyDown( const FGeometry& MyGeometry, const FKeyEvent& InKeyEvent )
 {
    // Pressing escape returns as if the user clicked cancel
diff --git a/Engine/Source/Editor/SourceControlWindows/Private/SSourceControlSubmit.h b/Engine/Source/Editor/SourceControlWindows/Private/SSourceControlSubmit.h
--- a/Engine/Source/Editor/SourceControlWindows/Private/SSourceControlSubmit.h
+++ b/Engine/Source/Editor/SourceControlWindows/Private/SSourceControlSubmit.h
@@ -72,6 +72,9 @@ public:
 	/** true if the item is enabled in the list */
 	bool IsEnabled() const { return !Item->IsConflicted() && Item->IsCurrent(); }
 
+	/** true if the item is source controlled and not marked for add nor for delete */
+	bool CanDiff() const { return Item->IsSourceControlled() && !Item->IsAdded() && !Item->IsDeleted(); }
+
 private:
 	/** Shared pointer to the source control state object itself */
 	FSourceControlStateRef Item;
@@ -184,6 +187,12 @@ private:
 	 */
 	void SortTree();
 
+	TSharedPtr<SWidget> OnCreateContextMenu();
+
+	bool CanDiffAgainstDepot() const;
+	void OnDiffAgainstDepot();
+	void OnDiffAgainstDepotSelected(TSharedPtr<FSubmitItem> InSelectedItem);
+	
 private:
 	ESubmitResults::Type DialogResult;
 
