PR #5268: Source Control History Window do not display Changelist title if not available (Contributed by [REDACTED_PERSON])

Switched the visibility to SelfHitTestInvisible

diff --git a/Engine/Source/Editor/SourceControlWindows/Private/SSourceControlHistory.cpp b/Engine/Source/Editor/SourceControlWindows/Private/SSourceControlHistory.cpp
--- a/Engine/Source/Editor/SourceControlWindows/Private/SSourceControlHistory.cpp
+++ b/Engine/Source/Editor/SourceControlWindows/Private/SSourceControlHistory.cpp
@@ -816,6 +816,7 @@ private:
 					[
 						SNew(STextBlock)
 						.Text(NSLOCTEXT("SourceControl.HistoryPanel.Info", "Changelist", "Changelist:"))
+						.Visibility(ISourceControlModule::Get().GetProvider().UsesChangelists() ? EVisibility::SelfHitTestInvisible : EVisibility::Hidden)
 					]
 					+SVerticalBox::Slot()
 					.FillHeight(0.25f)
