Made the CL description and info fields selectable in the editors revision history widget.

diff --git a/Engine/Source/Editor/SourceControlWindows/Private/SSourceControlHistory.cpp b/Engine/Source/Editor/SourceControlWindows/Private/SSourceControlHistory.cpp
--- a/Engine/Source/Editor/SourceControlWindows/Private/SSourceControlHistory.cpp
+++ b/Engine/Source/Editor/SourceControlWindows/Private/SSourceControlHistory.cpp
@@ -35,6 +35,7 @@
 #include "Widgets/Text/STextBlock.h"
 #include "Widgets/Layout/SBox.h"
 #include "Framework/MultiBox/MultiBoxBuilder.h"
+#include "Widgets/Input/SEditableTextBox.h"
 #include "Widgets/Layout/SSplitter.h"
 #include "Widgets/Views/SExpanderArrow.h"
 #include "Widgets/Views/SHeaderRow.h"
@@ -779,28 +780,32 @@ private:
 					.FillHeight(0.25f)
 					.Padding(Padding)
 					[
-						SNew(STextBlock)
+						SNew(SEditableText)
+						.IsReadOnly(true)
 						.Text(this, &SSourceControlHistoryWidget::GetRevisionNumber)	
 					]
 					+SVerticalBox::Slot()
 					.FillHeight(0.25f)
 					.Padding(Padding)
 					[
-						SNew(STextBlock)
+						SNew(SEditableText)
+						.IsReadOnly(true)
 						.Text(this, &SSourceControlHistoryWidget::GetDate)	
 					]
 					+SVerticalBox::Slot()
 					.FillHeight(0.25f)
 					.Padding(Padding)
 					[
-						SNew(STextBlock)
+						SNew(SEditableText)
+						.IsReadOnly(true)
 						.Text(this, &SSourceControlHistoryWidget::GetUserName)	
 					]
 					+SVerticalBox::Slot()
 					.FillHeight(0.25f)
 					.Padding(Padding)
 					[
-						SNew(STextBlock)
+						SNew(SEditableText)
+						.IsReadOnly(true)
 						.Text(this, &SSourceControlHistoryWidget::GetAction)	
 					]
 				]
@@ -850,28 +855,32 @@ private:
 					.FillHeight(0.25f)
 					.Padding(Padding)
 					[
-						SNew(STextBlock)
+						SNew(SEditableText)
+						.IsReadOnly(true)
 						.Text(this, &SSourceControlHistoryWidget::GetChangelistNumber)	
 					]
 					+SVerticalBox::Slot()
 					.FillHeight(0.25f)
 					.Padding(Padding)
 					[
-						SNew(STextBlock)
+						SNew(SEditableText)
+						.IsReadOnly(true)
 						.Text(this, &SSourceControlHistoryWidget::GetClientSpec)	
 					]
 					+SVerticalBox::Slot()
 					.FillHeight(0.25f)
 					.Padding(Padding)
 					[
-						SNew(STextBlock)
+						SNew(SEditableText)
+						.IsReadOnly(true)
 						.Text(this, &SSourceControlHistoryWidget::GetFileSize)	
 					]
 					+SVerticalBox::Slot()
 					.FillHeight(0.25f)
 					.Padding(Padding)
 					[
-						SNew(STextBlock)
+						SNew(SEditableText)
+						.IsReadOnly(true)
 						.Text(this, &SSourceControlHistoryWidget::GetBranchedFrom)	
 					]
 				]
@@ -893,7 +902,8 @@ private:
 					.AutoHeight()
 					.Padding(5)
 					[
-						SNew(STextBlock)
+						SNew(SEditableText)
+						.IsReadOnly(true)
 						.Text(this, &SSourceControlHistoryWidget::GetDescription)	
 					]
 				]
