[NO NON-PII COMMIT MESSAGE]

diff --git a/Engine/Source/Developer/Merge/Private/MergeUtils.cpp b/Engine/Source/Developer/Merge/Private/MergeUtils.cpp
--- a/Engine/Source/Developer/Merge/Private/MergeUtils.cpp
+++ b/Engine/Source/Developer/Merge/Private/MergeUtils.cpp
@@ -43,7 +43,7 @@ UObject const* FMergeToolUtils::LoadRevision(const FString& AssetName, const ISo
 	if (DesiredRevision.Get(TempFileName))
 	{
 		// Try and load that package
-		UPackage* TempPackage = LoadPackage(NULL, *TempFileName, LOAD_DisableCompileOnLoad);
+		UPackage* TempPackage = LoadPackage(NULL, *TempFileName, LOAD_ForDiff|LOAD_DisableCompileOnLoad);
 		if (TempPackage != NULL)
 		{
 			// Grab the old asset from that old package
diff --git a/Engine/Source/Editor/Kismet/Private/SBlueprintEditorToolbar.cpp b/Engine/Source/Editor/Kismet/Private/SBlueprintEditorToolbar.cpp
--- a/Engine/Source/Editor/Kismet/Private/SBlueprintEditorToolbar.cpp
+++ b/Engine/Source/Editor/Kismet/Private/SBlueprintEditorToolbar.cpp
@@ -243,7 +243,7 @@ static void OnDiffRevisionPicked(FRevisionInfo const& RevisionInfo, TWeakObjectP
 					if (Revision->Get(PreviousTempPkgName))
 					{
 						// Try and load that package
-						UPackage* PreviousTempPkg = LoadPackage(NULL, *PreviousTempPkgName, LOAD_DisableCompileOnLoad);
+						UPackage* PreviousTempPkg = LoadPackage(NULL, *PreviousTempPkgName, LOAD_ForDiff|LOAD_DisableCompileOnLoad);
 
 						if (PreviousTempPkg != NULL)
 						{
diff --git a/Engine/Source/Editor/SourceControlWindows/Private/SSourceControlHistory.cpp b/Engine/Source/Editor/SourceControlWindows/Private/SSourceControlHistory.cpp
--- a/Engine/Source/Editor/SourceControlWindows/Private/SSourceControlHistory.cpp
+++ b/Engine/Source/Editor/SourceControlWindows/Private/SSourceControlHistory.cpp
@@ -209,7 +209,7 @@ static UObject* GetAssetRevisionObject(TSharedPtr<FHistoryTreeItem> HistoryTreeI
 				if (FileRevision.IsValid() && FileRevision->Get(TempPackageName)) // grab the path to a temporary package (where the revision item will be stored)
 				{
 					// try and load the temporary package
-					AssetPackage = LoadPackage(NULL, *TempPackageName, LOAD_DisableCompileOnLoad);
+					AssetPackage = LoadPackage(NULL, *TempPackageName, LOAD_ForDiff|LOAD_DisableCompileOnLoad);
 				}
 			} // if FileSourceControlState.IsValid()
 		}
diff --git a/Engine/Source/Editor/WorldBrowser/Private/LevelCollectionModel.cpp b/Engine/Source/Editor/WorldBrowser/Private/LevelCollectionModel.cpp
--- a/Engine/Source/Editor/WorldBrowser/Private/LevelCollectionModel.cpp
+++ b/Engine/Source/Editor/WorldBrowser/Private/LevelCollectionModel.cpp
@@ -1258,7 +1258,7 @@ void FLevelCollectionModel::SCCDiffAgainstDepot(const FLevelModelList& InList, U
 					{
 						// Try and load that package
 						FText NotMapReason;
-						UPackage* OldPackage = LoadPackage(NULL, *TempFileName, LOAD_DisableCompileOnLoad);
+						UPackage* OldPackage = LoadPackage(NULL, *TempFileName, LOAD_ForDiff|LOAD_DisableCompileOnLoad);
 						if(OldPackage != NULL && InEditor->PackageIsAMapFile(*TempFileName, NotMapReason))
 						{
 							/* Set the revision information*/
