Allowed for the source control descriptions to not be cleared after a failed submission.

diff --git a/Engine/Source/Editor/SourceControlWindows/Private/SSourceControlSubmit.cpp b/Engine/Source/Editor/SourceControlWindows/Private/SSourceControlSubmit.cpp
--- a/Engine/Source/Editor/SourceControlWindows/Private/SSourceControlSubmit.cpp
+++ b/Engine/Source/Editor/SourceControlWindows/Private/SSourceControlSubmit.cpp
@@ -392,8 +392,6 @@ void SSourceControlSubmitWidget::FillChangeListDescription(FChangeListDescriptio
 {
 	OutDesc.Description = ChangeListDescriptionTextCtrl->GetText();
 
-	ChangeListDescriptionTextCtrl->SetText(FText()); // Don't save description on Submit
-
 	OutDesc.FilesForAdd.Empty();
 	OutDesc.FilesForSubmit.Empty();
 
@@ -419,6 +417,10 @@ bool SSourceControlSubmitWidget::WantToKeepCheckedOut()
 	return KeepCheckedOut == ECheckBoxState::Checked ? true : false;
 }
 
+void SSourceControlSubmitWidget::ClearChangeListDescription()
+{
+	ChangeListDescriptionTextCtrl->SetText(FText());
+}
 
 FReply SSourceControlSubmitWidget::OKClicked()
 {
diff --git a/Engine/Source/Editor/SourceControlWindows/Private/SSourceControlSubmit.h b/Engine/Source/Editor/SourceControlWindows/Private/SSourceControlSubmit.h
--- a/Engine/Source/Editor/SourceControlWindows/Private/SSourceControlSubmit.h
+++ b/Engine/Source/Editor/SourceControlWindows/Private/SSourceControlSubmit.h
@@ -118,6 +118,10 @@ public:
 
 	/** Does the user want to keep the files checked out */
 	bool WantToKeepCheckedOut();
+
+	/** Clears the current change list description */
+	void ClearChangeListDescription();
+
 private:
 	/**
 	 * @return the desired toggle state for the ToggleSelectedCheckBox.
diff --git a/Engine/Source/Editor/SourceControlWindows/Private/SourceControlWindows.cpp b/Engine/Source/Editor/SourceControlWindows/Private/SourceControlWindows.cpp
--- a/Engine/Source/Editor/SourceControlWindows/Private/SourceControlWindows.cpp
+++ b/Engine/Source/Editor/SourceControlWindows/Private/SourceControlWindows.cpp
@@ -322,6 +322,8 @@ bool FSourceControlWindows::PromptForCheckin(FCheckinResultInfo& OutResultInfo,
 		}
 	}
 
+	SourceControlWidget->ClearChangeListDescription();
+
 	return true;
 }
 
