PR #5789: Save changelist description when Submit to Source Control Window is canceled (Contributed by [REDACTED_PERSON])

diff --git a/Engine/Source/Editor/SourceControlWindows/Private/SSourceControlSubmit.cpp b/Engine/Source/Editor/SourceControlWindows/Private/SSourceControlSubmit.cpp
--- a/Engine/Source/Editor/SourceControlWindows/Private/SSourceControlSubmit.cpp
+++ b/Engine/Source/Editor/SourceControlWindows/Private/SSourceControlSubmit.cpp
@@ -69,6 +69,12 @@ TSharedRef<SWidget> SSourceControlSubmitListRow::GenerateWidgetForColumn(const F
 	return SNullWidget::NullWidget;
 }
 
+FText SSourceControlSubmitWidget::SavedChangeListDescription;
+
+SSourceControlSubmitWidget::~SSourceControlSubmitWidget()
+{
+	SavedChangeListDescription = ChangeListDescriptionTextCtrl->GetText();
+}
 
 void SSourceControlSubmitWidget::Construct(const FArguments& InArgs)
 {
@@ -133,6 +139,7 @@ void SSourceControlSubmitWidget::Construct(const FArguments& InArgs)
 				[
 					SAssignNew( ChangeListDescriptionTextCtrl, SMultiLineEditableTextBox )
 					.SelectAllTextWhenFocused( true )
+					.Text(SavedChangeListDescription)
 					.AutoWrapText( true )
 				]
 			]
@@ -415,6 +422,8 @@ FReply SSourceControlSubmitWidget::OKClicked()
 	DialogResult = ESubmitResults::SUBMIT_ACCEPTED;
 	ParentFrame.Pin()->RequestDestroyWindow();
 
+	ChangeListDescriptionTextCtrl->ClearContent(); // Don't save description on Submit
+
 	return FReply::Handled();
 }
 
diff --git a/Engine/Source/Editor/SourceControlWindows/Private/SSourceControlSubmit.h b/Engine/Source/Editor/SourceControlWindows/Private/SSourceControlSubmit.h
--- a/Engine/Source/Editor/SourceControlWindows/Private/SSourceControlSubmit.h
+++ b/Engine/Source/Editor/SourceControlWindows/Private/SSourceControlSubmit.h
@@ -99,10 +99,7 @@ public:
 
 	SLATE_END_ARGS()
 
-	/** Constructor */
-	SSourceControlSubmitWidget()
-	{
-	}
+	~SSourceControlSubmitWidget();
 
 	/** Constructs the widget */
 	void Construct(const FArguments& InArgs);
@@ -216,6 +213,9 @@ private:
 
 	/** Currently selected sorting mode */
 	EColumnSortMode::Type SortMode;
+
+	/** Submit Description saved when the widget is destroyed if canceled */
+	static FText SavedChangeListDescription;	
 };
 
 class SSourceControlSubmitListRow : public SMultiColumnTableRow<TSharedPtr<FSubmitItem>>
