PR #5804: Source Control History: bind Escape key to destroy the window (Quality of Life) (Contributed by [REDACTED_PERSON])

diff --git a/Engine/Source/Editor/SourceControlWindows/Private/SSourceControlHistory.cpp b/Engine/Source/Editor/SourceControlWindows/Private/SSourceControlHistory.cpp
--- a/Engine/Source/Editor/SourceControlWindows/Private/SSourceControlHistory.cpp
+++ b/Engine/Source/Editor/SourceControlWindows/Private/SSourceControlHistory.cpp
@@ -655,6 +655,7 @@ public:
 	void Construct( const FArguments& InArgs )
 	{	
 		AddHistoryInfo(InArgs._SourceControlStates.Get());
+		ParentWindow = InArgs._ParentWindow.Get();
 
 		TSharedRef<SHeaderRow> HeaderRow = SNew(SHeaderRow);
 
@@ -718,10 +719,25 @@ public:
 		{
 			MainHistoryListView->SetItemExpansion(HistoryCollection[i],true);
 		}
+
+		ParentWindow.Pin()->SetWidgetToFocusOnActivate(MainHistoryListView);
 	}
 
 private:
 
+	/** Used to intercept Escape key press, and interpret it as a close event */
+	virtual FReply OnKeyDown(const FGeometry& MyGeometry, const FKeyEvent& InKeyEvent) override
+	{
+		// Pressing escape returns as if the user closed the window
+		if (InKeyEvent.GetKey() == EKeys::Escape)
+		{
+			ParentWindow.Pin()->RequestDestroyWindow();
+			return FReply::Handled();
+		}
+
+		return FReply::Unhandled();
+	}
+
 	/**
 	 * Constructs the "Additional Info" panel that displays specific revision info
 	 */
@@ -1426,6 +1442,9 @@ private:
 
 	/** The last selected revision item; Displayed in the "additional information" subpanel */
 	TWeakPtr<FHistoryRevisionListViewItem> LastSelectedRevisionItem;
+
+	/** Pointer to the parent window */
+	TWeakPtr<SWindow> ParentWindow;
 };
 
 void FSourceControlWindows::DisplayRevisionHistory( const TArray<FString>& InPackageNames )
