PR #5863: Fix regression introduced by #5789: properly reset changelist description on Submit to Source Control (Contributed by [REDACTED_PERSON])

diff --git a/Engine/Source/Editor/SourceControlWindows/Private/SSourceControlSubmit.cpp b/Engine/Source/Editor/SourceControlWindows/Private/SSourceControlSubmit.cpp
--- a/Engine/Source/Editor/SourceControlWindows/Private/SSourceControlSubmit.cpp
+++ b/Engine/Source/Editor/SourceControlWindows/Private/SSourceControlSubmit.cpp
@@ -422,7 +422,7 @@ FReply SSourceControlSubmitWidget::OKClicked()
 	DialogResult = ESubmitResults::SUBMIT_ACCEPTED;
 	ParentFrame.Pin()->RequestDestroyWindow();
 
-	ChangeListDescriptionTextCtrl->ClearContent(); // Don't save description on Submit
+	ChangeListDescriptionTextCtrl->SetText(FText()); // Don't save description on Submit
 
 	return FReply::Handled();
 }
