Updated "Submit to Source Control" to include plugin content

diff --git a/Engine/Source/Editor/ContentBrowser/Private/FrontendFilters.cpp b/Engine/Source/Editor/ContentBrowser/Private/FrontendFilters.cpp
--- a/Engine/Source/Editor/ContentBrowser/Private/FrontendFilters.cpp
+++ b/Engine/Source/Editor/ContentBrowser/Private/FrontendFilters.cpp
@@ -9,6 +9,7 @@
 #include "ISourceControlModule.h"
 #include "SourceControlHelpers.h"
 #include "SourceControlOperations.h"
+#include "SourceControlWindows.h"
 #include "Editor.h"
 #include "AssetToolsModule.h"
 #include "ICollectionManager.h"
@@ -853,9 +854,7 @@ void FFrontendFilter_NotSourceControlled::RequestStatus()
 		// Request the state of files at filter construction time to make sure files have the correct state for the filter
 		TSharedRef<FUpdateStatus, ESPMode::ThreadSafe> UpdateStatusOperation = ISourceControlOperation::Create<FUpdateStatus>();
 
-		TArray<FString> Filenames;
-		Filenames.Add(FPaths::ConvertRelativePathToFull(FPaths::EngineContentDir()));
-		Filenames.Add(FPaths::ConvertRelativePathToFull(FPaths::ProjectContentDir()));
+		TArray<FString> Filenames = FSourceControlWindows::GetSourceControlLocations(/*bContentOnly*/true);
 		UpdateStatusOperation->SetCheckingAllFiles(false);
 		SourceControlProvider.Execute(UpdateStatusOperation, Filenames, EConcurrency::Asynchronous, FSourceControlOperationComplete::CreateSP(this, &FFrontendFilter_NotSourceControlled::SourceControlOperationComplete));
 	}
diff --git a/Engine/Source/Editor/SourceControlWindows/Private/SourceControlWindows.cpp b/Engine/Source/Editor/SourceControlWindows/Private/SourceControlWindows.cpp
--- a/Engine/Source/Editor/SourceControlWindows/Private/SourceControlWindows.cpp
+++ b/Engine/Source/Editor/SourceControlWindows/Private/SourceControlWindows.cpp
@@ -37,6 +37,29 @@ FCheckinResultInfo::FCheckinResultInfo()
 
 TWeakPtr<SNotificationItem> FSourceControlWindows::ChoosePackagesToCheckInNotification;
 
+TArray<FString> FSourceControlWindows::GetSourceControlLocations(const bool bContentOnly)
+{
+	TArray<FString> SourceControlLocations;
+
+	{
+		TArray<FString> RootPaths;
+		FPackageName::QueryRootContentPaths(RootPaths);
+		for (const FString& RootPath : RootPaths)
+		{
+			const FString RootPathOnDisk = FPackageName::LongPackageNameToFilename(RootPath);
+			SourceControlLocations.Add(FPaths::ConvertRelativePathToFull(RootPathOnDisk));
+		}
+	}
+
+	if (!bContentOnly)
+	{
+		SourceControlLocations.Add(FPaths::ConvertRelativePathToFull(FPaths::ProjectConfigDir()));
+		SourceControlLocations.Add(FPaths::ConvertRelativePathToFull(FPaths::GetProjectFilePath()));
+	}
+
+	return SourceControlLocations;
+}
+
 bool FSourceControlWindows::ChoosePackagesToCheckIn(const FSourceControlWindowsOnCheckInComplete& OnCompleteDelegate)
 {
 	if (!ISourceControlModule::Get().IsEnabled())
@@ -66,11 +89,7 @@ bool FSourceControlWindows::ChoosePackagesToCheckIn(const FSourceControlWindowsO
 	// Start selection process...
 
 	// make sure we update the SCC status of all packages (this could take a long time, so we will run it as a background task)
-	TArray<FString> Filenames;
-	Filenames.Add(FPaths::ConvertRelativePathToFull(FPaths::EngineContentDir()));
-	Filenames.Add(FPaths::ConvertRelativePathToFull(FPaths::ProjectContentDir()));
-	Filenames.Add(FPaths::ConvertRelativePathToFull(FPaths::ProjectConfigDir()));
-	Filenames.Add(FPaths::ConvertRelativePathToFull(FPaths::GetProjectFilePath()));
+	TArray<FString> Filenames = GetSourceControlLocations();
 
 	ISourceControlProvider& SourceControlProvider = ISourceControlModule::Get().GetProvider();
 	FSourceControlOperationRef Operation = ISourceControlOperation::Create<FUpdateStatus>();
@@ -369,11 +388,7 @@ void FSourceControlWindows::ChoosePackagesToCheckInCompleted(const TArray<UPacka
 		return;
 	}
 
-	TArray<FString> PendingDeletePaths;
-	PendingDeletePaths.Add(FPaths::ConvertRelativePathToFull(FPaths::EngineContentDir()));
-	PendingDeletePaths.Add(FPaths::ConvertRelativePathToFull(FPaths::ProjectContentDir()));
-	PendingDeletePaths.Add(FPaths::ConvertRelativePathToFull(FPaths::ProjectConfigDir()));
-	PendingDeletePaths.Add(FPaths::ConvertRelativePathToFull(FPaths::GetProjectFilePath()));
+	TArray<FString> PendingDeletePaths = GetSourceControlLocations();
 
 	const bool bUseSourceControlStateCache = true;
 	PromptForCheckin(OutResultInfo, PackageNames, PendingDeletePaths, ConfigFiles, bUseSourceControlStateCache);
diff --git a/Engine/Source/Editor/SourceControlWindows/Public/SourceControlWindows.h b/Engine/Source/Editor/SourceControlWindows/Public/SourceControlWindows.h
--- a/Engine/Source/Editor/SourceControlWindows/Public/SourceControlWindows.h
+++ b/Engine/Source/Editor/SourceControlWindows/Public/SourceControlWindows.h
@@ -35,6 +35,13 @@ DECLARE_DELEGATE_OneParam(FSourceControlWindowsOnCheckInComplete, const FCheckin
 class SOURCECONTROLWINDOWS_API FSourceControlWindows
 {
 public:
+	/**
+	 * Get the list of files and directories that source control should check when looking for changes.
+	 *
+	 * @param	bContentOnly	True to only include content directories.
+	 */
+	static TArray<FString> GetSourceControlLocations(const bool bContentOnly = false);
+
 	/**
 	 * Opens a user dialog to choose packages to submit.
 	 *
