Create a Diff against Workspace option in the source control history right click context menu
Disable Diff against Previous revision if right clicking the oldest revision
Fix file items not doing anything when diffing against previous or diffing via drag/drop by making them behave as if they are the most recent revision at that location.

diff --git a/Engine/Source/Editor/SourceControlWindows/Private/SSourceControlHistory.cpp b/Engine/Source/Editor/SourceControlWindows/Private/SSourceControlHistory.cpp
--- a/Engine/Source/Editor/SourceControlWindows/Private/SSourceControlHistory.cpp
+++ b/Engine/Source/Editor/SourceControlWindows/Private/SSourceControlHistory.cpp
@@ -286,7 +286,15 @@ static bool CanDiffSelectedItems(TArray< TSharedPtr<FHistoryTreeItem> > const& S
 		{
 			ErrorTextOut = NSLOCTEXT("SourceControlHistory", "CannotDiffWithSelf", "You cannot diff a revision against itself.");
 		}
-		else 
+		else if (!FirstSelection->Parent.IsValid() && FirstSelection->RevisionListItem.IsValid() && FirstSelection->Children[0] == SecondSelection)
+		{
+			ErrorTextOut = NSLOCTEXT("SourceControlHistory", "CannotDiffWithSelf", "You cannot diff a revision against itself.");
+		}
+		else if (!SecondSelection->Parent.IsValid() && FirstSelection->RevisionListItem.IsValid() && SecondSelection->Children[0] == FirstSelection)
+		{
+			ErrorTextOut = NSLOCTEXT("SourceControlHistory", "CannotDiffWithSelf", "You cannot diff a revision against itself.");
+		}
+		else
 		{
 			// @TODO make sure the two selections match type (calling GetAssetRevisionObject() to compare class types is too slow)
 			bCanDiffSelected = true;
@@ -1090,7 +1098,7 @@ private:
 			FileItem->FileListItem = MakeShareable(new FHistoryFileListViewItem( SourceControlState->GetFilename() ));
 
 			// Add each file revision
-			for ( int HistoryIndex = 0; HistoryIndex < SourceControlState->GetHistorySize(); HistoryIndex++ )
+			for ( int32 HistoryIndex = 0; HistoryIndex < SourceControlState->GetHistorySize(); HistoryIndex++ )
 			{
 				TSharedPtr<ISourceControlRevision, ESPMode::ThreadSafe> Revision = SourceControlState->GetHistoryItem(HistoryIndex);
 				check(Revision.IsValid());
@@ -1099,6 +1107,12 @@ private:
 				FileItem->Children.Add(RevisionItem);
 				RevisionItem->Parent = FileItem;
 
+				// File Items other than the current workspace version masquarade as the latest version of that file
+				if (HistoryIndex == 0 && HistoryCollection.Num() > 0)
+				{
+					FileItem->RevisionListItem = RevisionItem->RevisionListItem;
+				}
+
 				// add branch items if we have one
 				if(Revision->GetBranchSource().IsValid())
 				{
@@ -1163,6 +1177,16 @@ private:
 				)
 		);
 
+		MenuBuilder.AddMenuEntry(
+			NSLOCTEXT("SourceControl.HistoryWindow.Menu", "DiffAgainstWorkspace", "Diff Against Workspace File"),
+			NSLOCTEXT("SourceControl.HistoryWindow.Menu", "DiffAgainstWorkspaceTooltip", "See changes between this revision and your version of the asset."),
+			FSlateIcon(),
+			FUIAction(
+				FExecuteAction::CreateSP(this, &SSourceControlHistoryWidget::OnDiffAgainstWorkspace),
+				FCanExecuteAction::CreateSP(this, &SSourceControlHistoryWidget::CanDiffAgainstWorkspace)
+			)
+		);
+
 		if (CanDiffSelected())
 		{
 			MenuBuilder.AddMenuEntry( 
@@ -1181,12 +1205,46 @@ private:
 	/** See if we should enabled the 'diff against previous' option */
 	bool CanDiffAgainstPreviousRev() const
 	{
-		// Only allow option if we selected one item and it was a revision, not a file entry
 		TArray< TSharedPtr<FHistoryTreeItem> > SelectedRevs = MainHistoryListView->GetSelectedItems();
-		return(	SelectedRevs.Num() == 1 && SelectedRevs[0].IsValid() );
+
+		// Only allow option if we selected one item
+		if (SelectedRevs.Num() == 1 && SelectedRevs[0].IsValid())
+		{
+			TSharedPtr<FHistoryTreeItem> SelectedItem = SelectedRevs[0];
+
+			if (SelectedItem->Parent.IsValid())
+			{
+				// If it is a revision then allow diffing against previous revision unless it is the original revision
+				TSharedPtr<FHistoryTreeItem> FileItem = SelectedItem->Parent.Pin();
+				check(FileItem.IsValid());
+
+				int32 RevIndex = FileItem->Children.Find(SelectedItem);
+				check(RevIndex != INDEX_NONE);
+				if (RevIndex == FileItem->Children.Num()-1) // If oldest revision of this file
+				{
+					int32 FileIndex = HistoryCollection.Find(FileItem);
+					check(FileIndex != INDEX_NONE);
+					return (FileIndex < HistoryCollection.Num()-1);
+				}
+			}
+			else
+			{
+				// If it is a file item then allow diffing against previous revision unless it is the last
+				// file item and it only has one child
+				const int32 FileIndex = HistoryCollection.Find(SelectedItem);
+				if (FileIndex > 0 && FileIndex == HistoryCollection.Num() - 1)
+				{
+					return (SelectedItem->Children.Num() > 1);
+				}
+			}
+
+			return true;
+		}
+
+		return false;
 	}
 
-	/** Try and perfom a diff between the selected revision and the previous one */
+	/** Try and perform a diff between the selected revision and the previous one */
 	void OnDiffAgainstPreviousRev()
 	{
 		TArray< TSharedPtr<FHistoryTreeItem> > SelectedRevs = MainHistoryListView->GetSelectedItems();
@@ -1200,16 +1258,30 @@ private:
 
 			if (SelectedItem->RevisionListItem.IsValid())
 			{
-				TSharedPtr<FHistoryTreeItem> FileItem = SelectedItem->Parent.Pin();
-				check(FileItem.IsValid());
+				TSharedPtr<FHistoryTreeItem> FileItem;
+				int32 RevIndex;
+				
+				if (SelectedItem->Parent.IsValid())
+				{
+					FileItem = SelectedItem->Parent.Pin();
+					check(FileItem.IsValid());
+
+					// First, find index of selected revision in its parent file item
+					// NB. 0 is newest, increasing index means older
+					RevIndex = FileItem->Children.Find(SelectedItem);
+					check(RevIndex != INDEX_NONE);
+				}
+				else
+				{
+					// If the selected item doesn't have a parent, this is a file item
+					// masquerading as the most recent revision
+					FileItem = SelectedItem;
+					RevIndex = 0;
+				}
 
 				// Now we need to find previous revision
 				TSharedPtr<FHistoryTreeItem> PreRevisionItem;
 
-				// First, find index of selected revision in its parent file item
-				// NB. 0 is newest, increasing index means older
-				int32 RevIndex = FileItem->Children.Find(SelectedItem);
-				check(RevIndex != INDEX_NONE);
 				if(RevIndex == FileItem->Children.Num()-1) // If oldest revision of this file
 				{
 					// .. see if we have an older file
@@ -1264,6 +1336,59 @@ private:
 		}
 	}
 
+	/** See if we should enabled the 'diff against workspace' option */
+	bool CanDiffAgainstWorkspace() const
+	{
+		// Only allow option if we selected one item and it is not the current workspace row
+		TArray< TSharedPtr<FHistoryTreeItem> > SelectedRevs = MainHistoryListView->GetSelectedItems();
+		return (SelectedRevs.Num() == 1 && SelectedRevs[0].IsValid() && HistoryCollection[0] != SelectedRevs[0]);
+	}
+
+	/** Try and perform a diff between the selected revision and the current version of the asset */
+	void OnDiffAgainstWorkspace()
+	{
+		TArray< TSharedPtr<FHistoryTreeItem> > SelectedRevs = MainHistoryListView->GetSelectedItems();
+
+		if (SelectedRevs.Num() > 0 && SelectedRevs[0].IsValid())
+		{
+			FAssetToolsModule& AssetToolsModule = FModuleManager::LoadModuleChecked<FAssetToolsModule>(TEXT("AssetTools"));
+
+			TSharedPtr<FHistoryTreeItem> SelectedItem = SelectedRevs[0];
+			UObject* SelectedAsset = GetAssetRevisionObject(SelectedItem);
+
+			if (SelectedItem->RevisionListItem.IsValid())
+			{
+				UObject* CurrentAsset = (HistoryCollection.Num() > 0 ? GetAssetRevisionObject(HistoryCollection[0]) : nullptr);
+
+				if ((SelectedAsset != NULL) && (CurrentAsset != NULL))
+				{
+					FRevisionInfo OldRevisionInfo;
+					GetRevisionInfo(SelectedItem, OldRevisionInfo);
+					FRevisionInfo NewRevisionInfo;
+					NewRevisionInfo.Revision = TEXT("");
+
+					AssetToolsModule.Get().DiffAssets(SelectedAsset, CurrentAsset, OldRevisionInfo, NewRevisionInfo);
+				}
+				else
+				{
+					FMessageDialog::Open(EAppMsgType::Ok, NSLOCTEXT("SourceControl.HistoryWindow", "UnableToLoadAssets", "Unable to load assets to diff. Content may no longer be supported?"));
+				}
+			}
+			else if (SelectedAsset != NULL)
+			{
+				// this should be a file list-item (representing the current working version)
+				check(SelectedItem->FileListItem.IsValid());
+
+				FString const AssetName = SelectedAsset->GetName();
+				FString PackageName;
+				if (FPackageName::TryConvertFilenameToLongPackageName(SelectedItem->FileListItem->FileName, /*out*/ PackageName))
+				{
+					AssetToolsModule.Get().DiffAgainstDepot(SelectedAsset, PackageName, AssetName);
+				}
+			}
+		}
+	}
+
 	/**
 	 * Checks to see if the selected history-tree items can be diffed against each other.
 	 * 
@@ -1449,6 +1574,9 @@ private:
 
 void FSourceControlWindows::DisplayRevisionHistory( const TArray<FString>& InPackageNames )
 {
+	// Explicitly load the module so live coding will work with it
+	FModuleManager::Get().LoadModule(TEXT("SourceControlWindows"));
+
 	ISourceControlProvider& SourceControlProvider = ISourceControlModule::Get().GetProvider();
 
 	// Query for the file history for the provided packages
