Changelist Window initial implementation

diff --git a/Engine/Source/Editor/EditorStyle/Private/StarshipStyle.cpp b/Engine/Source/Editor/EditorStyle/Private/StarshipStyle.cpp
--- a/Engine/Source/Editor/EditorStyle/Private/StarshipStyle.cpp
+++ b/Engine/Source/Editor/EditorStyle/Private/StarshipStyle.cpp
@@ -7372,6 +7372,9 @@ void FStarshipEditorStyle::FStyle::SetupSourceControlStyles()
 		Set("SourceControl.StatusIcon.Off", new CORE_IMAGE_BRUSH_SVG("Starship/Common/reject", Icon16x16, FStyleColors::Foreground));
 		Set("SourceControl.StatusIcon.Unknown", new CORE_IMAGE_BRUSH_SVG("Starship/Common/help", Icon16x16, FStyleColors::AccentYellow));
 
+		Set("SourceControl.ChangelistsTab", new CORE_IMAGE_BRUSH_SVG("Starship/Common/check-circle", Icon16x16));
+		Set("SourceControl.Changelist", new IMAGE_BRUSH_SVG("Starship/AssetIcons/StaticMeshActor_64", Icon20x20, FStyleColors::AccentRed));
+
 		Set( "SourceControl.Add", new IMAGE_BRUSH( "Old/SourceControl/SCC_Action_Add",Icon10x10));
 		Set( "SourceControl.Edit", new IMAGE_BRUSH( "Old/SourceControl/SCC_Action_Edit",Icon10x10));
 		Set( "SourceControl.Delete", new IMAGE_BRUSH( "Old/SourceControl/SCC_Action_Delete",Icon10x10));
diff --git a/Engine/Source/Editor/SourceControlWindows/Private/SSourceControlChangelists.cpp b/Engine/Source/Editor/SourceControlWindows/Private/SSourceControlChangelists.cpp
new file mode 100644
--- /dev/null
+++ b/Engine/Source/Editor/SourceControlWindows/Private/SSourceControlChangelists.cpp
@@ -0,0 +1,364 @@
+// Copyright Epic Games, Inc. All Rights Reserved.
+
+#include "SSourceControlChangelists.h"
+
+#include "EditorStyleSet.h"
+
+#include "Widgets/DeclarativeSyntaxSupport.h"
+#include "Widgets/Images/SImage.h"
+#include "Widgets/Layout/SScrollBorder.h"
+#include "Widgets/SOverlay.h"
+
+#include "ISourceControlProvider.h"
+#include "ISourceControlModule.h"
+#include "SourceControlOperations.h"
+
+
+#define LOCTEXT_NAMESPACE "SourceControlChangelist"
+
+struct FChangelistTreeItem : public IChangelistTreeItem
+{
+	FChangelistTreeItem(FSourceControlChangelistStateRef InChangelistState)
+		: ChangelistState(InChangelistState)
+	{
+		Type = IChangelistTreeItem::Changelist;
+	}
+
+	FText GetDisplayText() const
+	{
+		return ChangelistState->GetDisplayText();
+	}
+
+	FText GetDescriptionText() const
+	{
+		return ChangelistState->GetDescriptionText();
+	}
+
+	FSourceControlChangelistStateRef ChangelistState;
+};
+
+struct FFileTreeItem : public IChangelistTreeItem
+{
+	FFileTreeItem(FSourceControlStateRef InFileState)
+		: FileState(InFileState)
+	{
+		Type = IChangelistTreeItem::File;
+	}
+
+	FText GetDisplayText() const
+	{
+		FString Filename = FileState->GetFilename();
+		return FText::FromString(Filename);
+	}
+
+	FSourceControlStateRef FileState;
+};
+
+SSourceControlChangelistsWidget::SSourceControlChangelistsWidget()
+{
+}
+
+void SSourceControlChangelistsWidget::Construct(const FArguments& InArgs)
+{
+	TreeView = CreateTreeviewWidget();
+
+	Refresh();
+
+	ChildSlot
+		[
+			SNew(SScrollBorder, TreeView.ToSharedRef())
+			[
+				TreeView.ToSharedRef()
+			]
+		];
+}
+
+void SSourceControlChangelistsWidget::Refresh()
+{
+	ISourceControlModule& SourceControlModule = ISourceControlModule::Get();
+	ISourceControlProvider& SourceControlProvider = SourceControlModule.GetProvider();
+
+	UpdatePendingChangelistsOperation = ISourceControlOperation::Create<FUpdatePendingChangelistsStatus>();
+	UpdatePendingChangelistsOperation->SetUpdateAllChangelists(true);
+	UpdatePendingChangelistsOperation->SetUpdateFilesStates(true);
+
+	SourceControlProvider.Execute(UpdatePendingChangelistsOperation.ToSharedRef(), EConcurrency::Asynchronous, FSourceControlOperationComplete::CreateSP(this, &SSourceControlChangelistsWidget::OnChangelistsStatusUpdated));
+}
+
+void SSourceControlChangelistsWidget::OnChangelistsStatusUpdated(const FSourceControlOperationRef& InOperation, ECommandResult::Type InType)
+{
+	ISourceControlModule& SourceControlModule = ISourceControlModule::Get();
+	ISourceControlProvider& SourceControlProvider = SourceControlModule.GetProvider();
+
+	TArray<FSourceControlChangelistStateRef> ChangelistsStates;
+	SourceControlProvider.GetState(SourceControlProvider.GetChangelists(EStateCacheUsage::Use), ChangelistsStates, EStateCacheUsage::Use);
+
+	ChangelistsNodes.Reset(ChangelistsStates.Num());
+
+	for (FSourceControlChangelistStateRef ChangelistState : ChangelistsStates)
+	{
+		FChangelistTreeItemRef ChangelistTreeItem = MakeShareable(new FChangelistTreeItem(ChangelistState));
+
+		for (FSourceControlStateRef FileRef : ChangelistState->GetFilesStates())
+		{
+			FChangelistTreeItemRef FileTreeItem = MakeShareable(new FFileTreeItem(FileRef));
+			ChangelistTreeItem->AddChild(FileTreeItem);
+		}
+
+		ChangelistsNodes.Add(ChangelistTreeItem);
+	}
+
+	TreeView->RequestTreeRefresh();
+}
+
+bool SSourceControlChangelistsWidget::OnIsSelectableOrNavigable(FChangelistTreeItemPtr InItem) const
+{
+	if (!InItem)
+	{
+		return false;
+	}
+
+	TArray<FChangelistTreeItemPtr> SelectedItems = TreeView->GetSelectedItems();
+
+	// Prevent selecting changelists and files at the same time.
+	if (!SelectedItems.IsEmpty() && SelectedItems[0]->GetTreeItemType() != InItem->GetTreeItemType())
+	{
+		return false;
+	}
+
+	return true;
+}
+
+TSharedRef<SChangelistTree> SSourceControlChangelistsWidget::CreateTreeviewWidget()
+{
+	return SAssignNew(TreeView, SChangelistTree)
+		.ItemHeight(24.0f)
+		.TreeItemsSource(&ChangelistsNodes)
+		.OnGenerateRow(this, &SSourceControlChangelistsWidget::OnGenerateRow)
+		.OnGetChildren(this, &SSourceControlChangelistsWidget::OnGetChildren)
+		.SelectionMode(ESelectionMode::Multi)
+		.OnIsSelectableOrNavigable(this, &SSourceControlChangelistsWidget::OnIsSelectableOrNavigable)
+		.HeaderRow
+		(
+			SNew(SHeaderRow)
+			+ SHeaderRow::Column("Change")
+			.DefaultLabel(LOCTEXT("Change", "Change"))
+			.FillWidth(0.2f)
+			+ SHeaderRow::Column("Description")
+			.DefaultLabel(LOCTEXT("Description", "Description"))
+		);
+}
+
+
+class SChangelistTableRow : public SMultiColumnTableRow<FChangelistTreeItemPtr>
+{
+public:
+	SLATE_BEGIN_ARGS(SChangelistTableRow)
+		: _TreeItemToVisualize()
+	{}
+	SLATE_ARGUMENT(FChangelistTreeItemPtr, TreeItemToVisualize)
+	SLATE_END_ARGS()
+
+public:
+	/**
+	* Construct child widgets that comprise this widget.
+	*
+	* @param InArgs Declaration from which to construct this widget.
+	*/
+	void Construct(const FArguments& InArgs, const TSharedRef<STableViewBase>& InOwner)
+	{
+		TreeItem = static_cast<FChangelistTreeItem*>(InArgs._TreeItemToVisualize.Get());
+
+		auto Args = FSuperRowType::FArguments();
+		SMultiColumnTableRow<FChangelistTreeItemPtr>::Construct(Args, InOwner);
+	}
+
+	// SMultiColumnTableRow overrides
+	virtual TSharedRef<SWidget> GenerateWidgetForColumn(const FName& ColumnName) override
+	{
+		if (ColumnName == TEXT("Change"))
+		{
+			const FSlateBrush* IconBrush = FEditorStyle::GetBrush("SourceControl.Changelist");
+
+			return SNew(SHorizontalBox)
+				+ SHorizontalBox::Slot()
+					.AutoWidth()
+					.VAlign(VAlign_Center)
+					[
+						SNew(SExpanderArrow, SharedThis(this))
+					]
+
+				+ SHorizontalBox::Slot()
+					.AutoWidth()
+					.VAlign(VAlign_Center)
+					[
+						SNew(SImage)
+						.Image(IconBrush)
+					]
+
+				+ SHorizontalBox::Slot()
+					.AutoWidth()
+					.Padding(2.0f, 0.0f)
+					.VAlign(VAlign_Center)
+					[
+						SNew(STextBlock)
+						.Text(this, &SChangelistTableRow::GetChangelistText)
+					];
+		}
+		else if (ColumnName == TEXT("Description"))
+		{
+			return SNew(SHorizontalBox)
+				+ SHorizontalBox::Slot()
+				.AutoWidth()
+				.Padding(2.0f, 0.0f)
+				.VAlign(VAlign_Center)
+				[
+					SNew(STextBlock)
+					.Text(this, &SChangelistTableRow::GetChangelistDescriptionText)
+				];
+		}
+		else
+		{
+			return SNullWidget::NullWidget;
+		}
+	}
+
+	FText GetChangelistText() const
+	{
+		return TreeItem->GetDisplayText();
+	}
+
+	FText GetChangelistDescriptionText() const
+	{
+		FString DescriptionString = TreeItem->GetDescriptionText().ToString();
+		DescriptionString.ReplaceInline(TEXT("\n"), TEXT(" "));
+		DescriptionString.TrimEndInline();
+		return FText::FromString(DescriptionString);
+	}
+
+private:
+	/** The info about the widget that we are visualizing. */
+	FChangelistTreeItem* TreeItem;
+};
+
+class SFileTableRow : public STableRow<FChangelistTreeItemPtr>
+{
+public:
+	SLATE_BEGIN_ARGS(SFileTableRow)
+		: _TreeItemToVisualize()
+	{}
+	SLATE_ARGUMENT(FChangelistTreeItemPtr, TreeItemToVisualize)
+	SLATE_END_ARGS()
+
+public:
+	/**
+	* Construct child widgets that comprise this widget.
+	*
+	* @param InArgs Declaration from which to construct this widget.
+	*/
+	void Construct(const FArguments& InArgs, const TSharedRef<STableViewBase>& InOwner)
+	{
+		TreeItem = static_cast<FFileTreeItem*>(InArgs._TreeItemToVisualize.Get());
+
+		const FSlateBrush* IconBrush = FEditorStyle::GetBrush("ContentBrowser.ColumnViewAssetIcon");
+
+		// Make icon overlays (eg, SCC and dirty status) a reasonable size in relation to the icon size (note: it is assumed this icon is square)
+		const float ICON_SCALING_FACTOR = 0.7f;
+		const float IconOverlaySize = IconBrush->ImageSize.X * ICON_SCALING_FACTOR;
+
+		ChildSlot
+		[
+			SNew(SHorizontalBox)
+
+			// Icon
+			+SHorizontalBox::Slot()
+				.AutoWidth()
+				.Padding(40, 0, 4, 0)
+				[
+					SNew(SOverlay)
+				
+					// The actual icon
+					+SOverlay::Slot()
+					[
+						SNew(SImage)
+						.Image( IconBrush )
+					]
+
+					// Source control state
+					+SOverlay::Slot()
+					.HAlign(HAlign_Left)
+					.VAlign(VAlign_Top)
+					[
+						SNew(SBox)
+						.WidthOverride(IconOverlaySize)
+						.HeightOverride(IconOverlaySize)
+						[
+							SNew(SImage)
+							.Image(this, &SFileTableRow::GetSCCStateImage)
+						]
+					]
+				]
+
+			+ SHorizontalBox::Slot()
+				.AutoWidth()
+				.VAlign(VAlign_Center)
+				[
+					SNew(STextBlock)
+					.Text(this, &SFileTableRow::GetDisplayText)
+				]
+		];
+
+		STableRow<FChangelistTreeItemPtr>::ConstructInternal(
+			STableRow::FArguments()
+			.ShowSelection(true),
+			InOwner
+		);
+	}
+
+	FText GetDisplayText() const
+	{
+		return TreeItem->GetDisplayText();
+	}
+
+	const FSlateBrush* GetSCCStateImage() const
+	{
+		return FEditorStyle::GetBrush(TreeItem->FileState->GetSmallIconName());
+	}
+
+private:
+	/** The info about the widget that we are visualizing. */
+	FFileTreeItem* TreeItem;
+};
+
+TSharedRef<ITableRow> SSourceControlChangelistsWidget::OnGenerateRow(FChangelistTreeItemPtr InTreeItem, const TSharedRef<STableViewBase>& OwnerTable)
+{
+	switch (InTreeItem->GetTreeItemType())
+	{
+	case IChangelistTreeItem::Changelist:
+		return SNew(SChangelistTableRow, OwnerTable)
+			.TreeItemToVisualize(InTreeItem);
+
+	case IChangelistTreeItem::File:
+		return SNew(SFileTableRow, OwnerTable)
+			.TreeItemToVisualize(InTreeItem);
+
+	default:
+		check(false);
+	};
+
+	return SNew(STableRow<TSharedPtr<FString>>, OwnerTable);
+}
+
+void SSourceControlChangelistsWidget::OnGetChildren(FChangelistTreeItemPtr InParent, TArray<FChangelistTreeItemPtr>& OutChildren)
+{
+	for (auto& Child : InParent->GetChildren())
+	{
+		// Should never have bogus entries in this list
+		check(Child.IsValid());
+		OutChildren.Add(Child);
+	}
+}
+
+
+
+#undef LOCTEXT_NAMESPACE
\ No newline at end of file
diff --git a/Engine/Source/Editor/SourceControlWindows/Private/SSourceControlChangelists.h b/Engine/Source/Editor/SourceControlWindows/Private/SSourceControlChangelists.h
new file mode 100644
--- /dev/null
+++ b/Engine/Source/Editor/SourceControlWindows/Private/SSourceControlChangelists.h
@@ -0,0 +1,105 @@
+// Copyright Epic Games, Inc. All Rights Reserved.
+
+#pragma once
+
+#include "CoreMinimal.h"
+
+#include "Widgets/SCompoundWidget.h"
+#include "Widgets/Views/STreeView.h"
+
+#include "ISourceControlProvider.h"
+
+struct IChangelistTreeItem;
+typedef TSharedPtr<IChangelistTreeItem> FChangelistTreeItemPtr;
+typedef TSharedRef<IChangelistTreeItem> FChangelistTreeItemRef;
+
+struct IChangelistTreeItem : TSharedFromThis<IChangelistTreeItem>
+{
+	enum TreeItemType
+	{
+		Invalid,
+		Changelist,
+		File
+	};
+
+	/** Get this item's parent. Can be nullptr. */
+	FChangelistTreeItemPtr GetParent() const
+	{
+		return Parent;
+	}
+
+	/** Get this item's children, if any. Although we store as weak pointers, they are guaranteed to be valid. */
+	const TArray<FChangelistTreeItemPtr>& GetChildren() const
+	{
+		return Children;
+	}
+
+	/** Returns the TreeItem's type */
+	const TreeItemType GetTreeItemType() const
+	{
+		return Type;
+	}
+
+	/** Add a child to this item */
+	void AddChild(FChangelistTreeItemRef Child)
+	{
+		Child->Parent = AsShared();
+		Children.Add(MoveTemp(Child));
+	}
+
+	/** Remove a child from this item */
+	void RemoveChild(const FChangelistTreeItemRef& Child)
+	{
+		if (Children.Remove(Child))
+		{
+			Child->Parent = nullptr;
+		}
+	}
+
+protected:
+	/** This item's parent, if any. */
+	FChangelistTreeItemPtr Parent;
+
+	/** Array of children contained underneath this item */
+	TArray<FChangelistTreeItemPtr> Children;
+
+	TreeItemType Type;
+};
+
+typedef STreeView<FChangelistTreeItemPtr> SChangelistTree;
+
+class SSourceControlChangelistsWidget : public SCompoundWidget
+{
+public:
+	SLATE_BEGIN_ARGS(SSourceControlChangelistsWidget) {}
+	SLATE_END_ARGS()
+
+	/**
+	* Constructor.
+	*/
+	SSourceControlChangelistsWidget();
+
+	/** Constructs the widget */
+	void Construct(const FArguments& InArgs);
+
+private:
+	TSharedRef<SChangelistTree> CreateTreeviewWidget();
+
+	TSharedRef<ITableRow> OnGenerateRow(FChangelistTreeItemPtr InTreeItem, const TSharedRef<STableViewBase>& OwnerTable);
+	void OnGetChildren(FChangelistTreeItemPtr InParent, TArray<FChangelistTreeItemPtr>& OutChildren);
+
+	bool OnIsSelectableOrNavigable(FChangelistTreeItemPtr InItem) const;
+
+	void Refresh();
+	void OnChangelistsStatusUpdated(const FSourceControlOperationRef&, ECommandResult::Type);
+
+private:
+
+	/** Changelists (root nodes) */
+	TArray<FChangelistTreeItemPtr> ChangelistsNodes;
+
+	/** Changelist treeview widget */
+	TSharedPtr<SChangelistTree> TreeView;
+
+	TSharedPtr<class FUpdatePendingChangelistsStatus, ESPMode::ThreadSafe> UpdatePendingChangelistsOperation;
+};
diff --git a/Engine/Source/Editor/SourceControlWindows/Private/SourceControlWindows.cpp b/Engine/Source/Editor/SourceControlWindows/Private/SourceControlWindows.cpp
--- a/Engine/Source/Editor/SourceControlWindows/Private/SourceControlWindows.cpp
+++ b/Engine/Source/Editor/SourceControlWindows/Private/SourceControlWindows.cpp
@@ -15,8 +15,6 @@
 #include "Widgets/Notifications/SNotificationList.h"
 
 
-IMPLEMENT_MODULE( FDefaultModuleImpl, SourceControlWindows );
-
 #if SOURCE_CONTROL_WITH_SLATE
 
 #define LOCTEXT_NAMESPACE "SourceControlWindows"
diff --git a/Engine/Source/Editor/SourceControlWindows/Private/SourceControlWindowsModule.cpp b/Engine/Source/Editor/SourceControlWindows/Private/SourceControlWindowsModule.cpp
new file mode 100644
--- /dev/null
+++ b/Engine/Source/Editor/SourceControlWindows/Private/SourceControlWindowsModule.cpp
@@ -0,0 +1,109 @@
+// Copyright Epic Games, Inc. All Rights Reserved.
+
+#include "CoreMinimal.h"
+#include "Modules/ModuleManager.h"
+#include "ISourceControlWindowsModule.h"
+
+#include "Widgets/Docking/SDockTab.h"
+#include "Textures/SlateIcon.h"
+#include "Framework/Application/SlateApplication.h"
+#include "EditorStyleSet.h"
+#include "LevelEditor.h"
+
+#include "WorkspaceMenuStructure.h"
+#include "WorkspaceMenuStructureModule.h"
+
+#include "SSourceControlChangelists.h"
+
+#define LOCTEXT_NAMESPACE "SourceControlWindows"
+
+/**
+ * SourceControlWindows module
+ */
+class FSourceControlWindowsModule : public ISourceControlWindowsModule
+{
+public:
+	/**
+	 * Called right after the module DLL has been loaded and the module object has been created
+	 */
+	virtual void StartupModule() override;
+
+	/**
+	 * Called before the module is unloaded, right before the module object is destroyed.
+	 */
+	virtual void ShutdownModule() override;
+
+	virtual void ShowChangelistsTab() override;
+
+private:
+	TSharedRef<SDockTab> CreateChangelistsTab(const FSpawnTabArgs& Args);
+	TSharedPtr<SWidget> CreateChangelistsUI();
+
+private:
+	TWeakPtr<SDockTab> ChangelistsTab;
+	TWeakPtr<SSourceControlChangelistsWidget> ChangelistsWidget;
+};
+
+IMPLEMENT_MODULE(FSourceControlWindowsModule, SourceControlWindows);
+
+static const FName SourceControlChangelistsTabName = FName(TEXT("SourceControlChangelists"));
+
+void FSourceControlWindowsModule::StartupModule()
+{
+	ISourceControlWindowsModule::StartupModule();
+
+	FGlobalTabmanager::Get()->RegisterNomadTabSpawner(SourceControlChangelistsTabName, FOnSpawnTab::CreateRaw(this, &FSourceControlWindowsModule::CreateChangelistsTab))
+		.SetDisplayName(LOCTEXT("ChangelistsTabTitle", "Changelists"));
+
+#if WITH_HOT_RELOAD
+	// This code attempts to relaunch the GameplayCueEditor tab when you hotreload this module
+	if (GIsHotReload && FSlateApplication::IsInitialized())
+	{
+		FLevelEditorModule& LevelEditorModule = FModuleManager::GetModuleChecked<FLevelEditorModule>(TEXT("LevelEditor"));
+		TSharedPtr<FTabManager> LevelEditorTabManager = LevelEditorModule.GetLevelEditorTabManager();
+		LevelEditorTabManager->TryInvokeTab(SourceControlChangelistsTabName);
+	}
+#endif // WITH_HOT_RELOAD
+}
+
+void FSourceControlWindowsModule::ShutdownModule()
+{
+	if (FSlateApplication::IsInitialized())
+	{
+		FGlobalTabmanager::Get()->UnregisterNomadTabSpawner(SourceControlChangelistsTabName);
+
+		if (ChangelistsTab.IsValid())
+		{
+			ChangelistsTab.Pin()->RequestCloseTab();
+		}
+	}
+}
+
+TSharedRef<SDockTab> FSourceControlWindowsModule::CreateChangelistsTab(const FSpawnTabArgs & Args)
+{
+	return SAssignNew(ChangelistsTab, SDockTab)
+		.Icon(FEditorStyle::Get().GetBrush("SourceControl.ChangelistsTab"))
+		.TabRole(ETabRole::NomadTab)
+		[
+			CreateChangelistsUI().ToSharedRef()
+		];
+}
+
+TSharedPtr<SWidget> FSourceControlWindowsModule::CreateChangelistsUI()
+{
+	TSharedPtr<SWidget> ReturnWidget;
+	if (IsInGameThread())
+	{
+		TSharedPtr<SSourceControlChangelistsWidget> SharedPtr = SNew(SSourceControlChangelistsWidget);
+		ReturnWidget = SharedPtr;
+		ChangelistsWidget = SharedPtr;
+	}
+	return ReturnWidget;
+}
+
+void FSourceControlWindowsModule::ShowChangelistsTab()
+{
+	FGlobalTabmanager::Get()->TryInvokeTab(FTabId(SourceControlChangelistsTabName));
+}
+
+#undef LOCTEXT_NAMESPACE
diff --git a/Engine/Source/Editor/SourceControlWindows/Public/ISourceControlWindowsModule.h b/Engine/Source/Editor/SourceControlWindows/Public/ISourceControlWindowsModule.h
new file mode 100644
--- /dev/null
+++ b/Engine/Source/Editor/SourceControlWindows/Public/ISourceControlWindowsModule.h
@@ -0,0 +1,23 @@
+// Copyright Epic Games, Inc. All Rights Reserved.
+#pragma once
+
+#include "CoreMinimal.h"
+#include "Modules/ModuleInterface.h"
+#include "Modules/ModuleManager.h"
+
+/**
+ * SourceControlWindows module interface
+ */
+class ISourceControlWindowsModule : public IModuleInterface
+{
+public:
+	/**
+	 * Get reference to the SourceControlWindows module instance
+	 */
+	static inline ISourceControlWindowsModule& Get()
+	{
+		return FModuleManager::LoadModuleChecked<ISourceControlWindowsModule>("SourceControlWindows");
+	}
+
+	virtual void ShowChangelistsTab() = 0;
+};
diff --git a/Engine/Source/Editor/SourceControlWindows/SourceControlWindows.Build.cs b/Engine/Source/Editor/SourceControlWindows/SourceControlWindows.Build.cs
--- a/Engine/Source/Editor/SourceControlWindows/SourceControlWindows.Build.cs
+++ b/Engine/Source/Editor/SourceControlWindows/SourceControlWindows.Build.cs
@@ -18,6 +18,7 @@ public class SourceControlWindows : ModuleRules
 				"SourceControl", 
 				"AssetTools",
 				"EditorFramework",
+				"WorkspaceMenuStructure",
                 "UnrealEd"			// We need this dependency here because we use PackageTools.
 			}
 		);
diff --git a/Engine/Source/Editor/StatusBar/Private/SourceControlMenuHelpers.cpp b/Engine/Source/Editor/StatusBar/Private/SourceControlMenuHelpers.cpp
--- a/Engine/Source/Editor/StatusBar/Private/SourceControlMenuHelpers.cpp
+++ b/Engine/Source/Editor/StatusBar/Private/SourceControlMenuHelpers.cpp
@@ -5,6 +5,7 @@
 #include "SourceControlOperations.h"
 #include "ISourceControlProvider.h"
 #include "ISourceControlModule.h"
+#include "ISourceControlWindowsModule.h"
 #include "SourceControlMenuHelpers.h"
 #include "SourceControlWindows.h"
 #include "FileHelpers.h"
@@ -35,12 +36,13 @@ void FSourceControlCommands::RegisterCommands()
 {
 	UI_COMMAND(ConnectToSourceControl, "Connect to Source Control...", "Opens a dialog to connect to source control.", EUserInterfaceActionType::Button, FInputChord());
 	UI_COMMAND(ChangeSourceControlSettings, "Change Source Control Settings...", "Opens a dialog to change source control settings.", EUserInterfaceActionType::Button, FInputChord());
+	UI_COMMAND(ViewChangelists, "View Changelists...", "Open the Source Control Changelist tab.", EUserInterfaceActionType::Button, FInputChord());
 	UI_COMMAND(CheckOutModifiedFiles, "Check Out Modified Files...", "Opens a dialog to check out any assets which have been modified.", EUserInterfaceActionType::Button, FInputChord());
 	UI_COMMAND(SubmitToSourceControl, "Submit to Source Control...", "Opens a dialog with check in options for content and levels.",   EUserInterfaceActionType::Button, FInputChord());
 
 	ActionList->MapAction(
 		ConnectToSourceControl,
-		FExecuteAction::CreateStatic(&FSourceControlCommands::ConnectToSourceControl_Clicked)                        
+		FExecuteAction::CreateStatic(&FSourceControlCommands::ConnectToSourceControl_Clicked)
 	);
 
 	ActionList->MapAction(
@@ -48,6 +50,12 @@ void FSourceControlCommands::RegisterCommands()
 		FExecuteAction::CreateStatic(&FSourceControlCommands::ConnectToSourceControl_Clicked)
 	);
 
+	ActionList->MapAction(
+		ViewChangelists,
+		FExecuteAction::CreateStatic(&FSourceControlCommands::ViewChangelists_Clicked),
+		FCanExecuteAction::CreateStatic(&FSourceControlCommands::ViewChangelists_CanExecute)
+	);
+
 	ActionList->MapAction(
 		CheckOutModifiedFiles,
 		FExecuteAction::CreateStatic(&FSourceControlCommands::CheckOutModifiedFiles_Clicked),
@@ -69,6 +77,23 @@ void FSourceControlCommands::ConnectToSourceControl_Clicked()
 	SourceControlModule.ShowLoginDialog(FSourceControlLoginClosed(), ELoginWindowMode::Modeless, EOnLoginWindowStartup::PreserveProvider);
 }
 
+bool FSourceControlCommands::ViewChangelists_CanExecute()
+{
+	ISourceControlModule& SourceControlModule = ISourceControlModule::Get();
+	if (ISourceControlModule::Get().IsEnabled() &&
+		ISourceControlModule::Get().GetProvider().IsAvailable())
+	{
+		return true;
+	}
+
+	return false;
+}
+
+void FSourceControlCommands::ViewChangelists_Clicked()
+{
+	ISourceControlWindowsModule::Get().ShowChangelistsTab();
+}
+
 bool FSourceControlCommands::CheckOutModifiedFiles_CanExecute()
 {
 	ISourceControlModule& SourceControlModule = ISourceControlModule::Get();
@@ -160,6 +185,13 @@ TSharedRef<SWidget> FSourceControlMenuHelpers::GenerateSourceControlMenuContent(
 
 	Section.AddSeparator("SourceControlConnectionSeparator");
 
+	Section.AddMenuEntry(
+		FSourceControlCommands::Get().ViewChangelists,
+		TAttribute<FText>(),
+		TAttribute<FText>(),
+		FSlateIcon(FEditorStyle::GetStyleSetName(), "SourceControl.ChangelistsTab")
+	);
+
 	Section.AddMenuEntry(
 		FSourceControlCommands::Get().CheckOutModifiedFiles,
 		TAttribute<FText>(),
diff --git a/Engine/Source/Editor/StatusBar/Private/SourceControlMenuHelpers.h b/Engine/Source/Editor/StatusBar/Private/SourceControlMenuHelpers.h
--- a/Engine/Source/Editor/StatusBar/Private/SourceControlMenuHelpers.h
+++ b/Engine/Source/Editor/StatusBar/Private/SourceControlMenuHelpers.h
@@ -22,6 +22,8 @@ public:
 private:
 
 	static void ConnectToSourceControl_Clicked();
+	static bool ViewChangelists_CanExecute();
+	static void ViewChangelists_Clicked();
 	static bool CheckOutModifiedFiles_CanExecute();
 	static void CheckOutModifiedFiles_Clicked();
 	static bool SubmitToSourceControl_CanExecute();
@@ -33,6 +35,7 @@ public:
 	 */
 	TSharedPtr< FUICommandInfo > ConnectToSourceControl;
 	TSharedPtr< FUICommandInfo > ChangeSourceControlSettings;
+	TSharedPtr< FUICommandInfo > ViewChangelists;
 	TSharedPtr< FUICommandInfo > CheckOutModifiedFiles;
 	TSharedPtr< FUICommandInfo > SubmitToSourceControl;
 
diff --git a/Engine/Source/Editor/UnrealEd/Private/EditorEngine.cpp b/Engine/Source/Editor/UnrealEd/Private/EditorEngine.cpp
--- a/Engine/Source/Editor/UnrealEd/Private/EditorEngine.cpp
+++ b/Engine/Source/Editor/UnrealEd/Private/EditorEngine.cpp
@@ -974,6 +974,7 @@ void UEditorEngine::Init(IEngineLoop* InEngineLoop)
 			TEXT("MainFrame"),
 			TEXT("OutputLog"),
 			TEXT("SourceControl"),
+			TEXT("SourceControlWindows"),
 			TEXT("TextureCompressor"),
 			TEXT("MeshUtilities"),
 			TEXT("MovieSceneTools"),
