Fix up deprecation warnings in SSourceControlChangelists widget

diff --git a/Engine/Source/Editor/SourceControlWindows/Private/SSourceControlChangelists.cpp b/Engine/Source/Editor/SourceControlWindows/Private/SSourceControlChangelists.cpp
--- a/Engine/Source/Editor/SourceControlWindows/Private/SSourceControlChangelists.cpp
+++ b/Engine/Source/Editor/SourceControlWindows/Private/SSourceControlChangelists.cpp
@@ -12,6 +12,7 @@
 #include "ISourceControlProvider.h"
 #include "ISourceControlModule.h"
 #include "SourceControlOperations.h"
+#include "Widgets/Images/SLayeredImage.h"
 
 
 #define LOCTEXT_NAMESPACE "SourceControlChangelist"
@@ -293,8 +294,7 @@ public:
 						.WidthOverride(IconOverlaySize)
 						.HeightOverride(IconOverlaySize)
 						[
-							SNew(SImage)
-							.Image(this, &SFileTableRow::GetSCCStateImage)
+							SNew(SLayeredImage, TreeItem->FileState->GetIcon())
 						]
 					]
 				]
@@ -320,11 +320,6 @@ public:
 		return TreeItem->GetDisplayText();
 	}
 
-	const FSlateBrush* GetSCCStateImage() const
-	{
-		return FEditorStyle::GetBrush(TreeItem->FileState->GetSmallIconName());
-	}
-
 private:
 	/** The info about the widget that we are visualizing. */
 	FFileTreeItem* TreeItem;
diff --git a/Engine/Source/Runtime/SlateCore/Private/Widgets/Images/SLayeredImage.cpp b/Engine/Source/Runtime/SlateCore/Private/Widgets/Images/SLayeredImage.cpp
--- a/Engine/Source/Runtime/SlateCore/Private/Widgets/Images/SLayeredImage.cpp
+++ b/Engine/Source/Runtime/SlateCore/Private/Widgets/Images/SLayeredImage.cpp
@@ -40,6 +40,13 @@ void SLayeredImage::Construct(const FArguments& InArgs, int32 NumLayers)
 	}
 }
 
+void SLayeredImage::Construct(const FArguments& InArgs, const FSlateIcon& InIcon)
+{
+	SImage::Construct(InArgs);
+
+	SetFromSlateIcon(InIcon);
+}
+
 int32 SLayeredImage::OnPaint(const FPaintArgs& Args, const FGeometry& AllottedGeometry, const FSlateRect& MyCullingRect, FSlateWindowElementList& OutDrawElements, int32 LayerId, const FWidgetStyle& InWidgetStyle, bool bParentEnabled) const
 {
 	// this will draw Image[0]:
diff --git a/Engine/Source/Runtime/SlateCore/Public/Widgets/Images/SLayeredImage.h b/Engine/Source/Runtime/SlateCore/Public/Widgets/Images/SLayeredImage.h
--- a/Engine/Source/Runtime/SlateCore/Public/Widgets/Images/SLayeredImage.h
+++ b/Engine/Source/Runtime/SlateCore/Public/Widgets/Images/SLayeredImage.h
@@ -25,6 +25,9 @@ public:
 	/** Constructor that adds NumLayers blank layers on top of the image for later use */
 	void Construct(const FArguments& InArgs, int32 NumLayers = 0);
 
+	/** Constructor that sets the base image and any overlay layers defined in the Slate icon */
+	void Construct(const FArguments& InArgs, const FSlateIcon& InIcon);
+
 	//~ Begin SWidget Interface
 	virtual int32 OnPaint(const FPaintArgs& Args, const FGeometry& AllottedGeometry, const FSlateRect& MyCullingRect, FSlateWindowElementList& OutDrawElements, int32 LayerId, const FWidgetStyle& InWidgetStyle, bool bParentEnabled) const override;
 	//~ End SWidget Interface
