Changelist window - Refresh changelists on provider change

diff --git a/Engine/Source/Editor/SourceControlWindows/Private/SSourceControlChangelists.cpp b/Engine/Source/Editor/SourceControlWindows/Private/SSourceControlChangelists.cpp
--- a/Engine/Source/Editor/SourceControlWindows/Private/SSourceControlChangelists.cpp
+++ b/Engine/Source/Editor/SourceControlWindows/Private/SSourceControlChangelists.cpp
@@ -61,6 +61,11 @@ SSourceControlChangelistsWidget::SSourceControlChangelistsWidget()
 
 void SSourceControlChangelistsWidget::Construct(const FArguments& InArgs)
 {
+	// Register delegates
+	ISourceControlModule& SCCModule = ISourceControlModule::Get();
+	SCCModule.RegisterProviderChanged(FSourceControlProviderChanged::FDelegate::CreateSP(this, &SSourceControlChangelistsWidget::OnSourceControlProviderChanged));
+	SourceControlStateChangedDelegateHandle = SCCModule.GetProvider().RegisterSourceControlStateChanged_Handle(FSourceControlStateChanged::FDelegate::CreateSP(this, &SSourceControlChangelistsWidget::OnSourceControlStateChanged));
+
 	TreeView = CreateTreeviewWidget();
 
 	Refresh();
@@ -76,37 +81,65 @@ void SSourceControlChangelistsWidget::Construct(const FArguments& InArgs)
 
 void SSourceControlChangelistsWidget::Refresh()
 {
-	ISourceControlModule& SourceControlModule = ISourceControlModule::Get();
-	ISourceControlProvider& SourceControlProvider = SourceControlModule.GetProvider();
+	if (ISourceControlModule::Get().IsEnabled())
+	{
+		UpdatePendingChangelistsOperation = ISourceControlOperation::Create<FUpdatePendingChangelistsStatus>();
+		UpdatePendingChangelistsOperation->SetUpdateAllChangelists(true);
+		UpdatePendingChangelistsOperation->SetUpdateFilesStates(true);
+
+		ISourceControlProvider& SourceControlProvider = ISourceControlModule::Get().GetProvider();
+		SourceControlProvider.Execute(UpdatePendingChangelistsOperation.ToSharedRef(), EConcurrency::Asynchronous, FSourceControlOperationComplete::CreateSP(this, &SSourceControlChangelistsWidget::OnChangelistsStatusUpdated));
+	}
+	else
+	{
+		// No provider available, clear changelist tree
+		ChangelistsNodes.Empty();
+		TreeView->RequestTreeRefresh();
+	}
+}
 
-	UpdatePendingChangelistsOperation = ISourceControlOperation::Create<FUpdatePendingChangelistsStatus>();
-	UpdatePendingChangelistsOperation->SetUpdateAllChangelists(true);
-	UpdatePendingChangelistsOperation->SetUpdateFilesStates(true);
+void SSourceControlChangelistsWidget::OnSourceControlProviderChanged(ISourceControlProvider& OldProvider, ISourceControlProvider& NewProvider)
+{
+	OldProvider.UnregisterSourceControlStateChanged_Handle(SourceControlStateChangedDelegateHandle);
+	SourceControlStateChangedDelegateHandle = NewProvider.RegisterSourceControlStateChanged_Handle(FSourceControlStateChanged::FDelegate::CreateSP(this, &SSourceControlChangelistsWidget::OnSourceControlStateChanged));
 
-	SourceControlProvider.Execute(UpdatePendingChangelistsOperation.ToSharedRef(), EConcurrency::Asynchronous, FSourceControlOperationComplete::CreateSP(this, &SSourceControlChangelistsWidget::OnChangelistsStatusUpdated));
+	OnSourceControlStateChanged();
+}
+
+void SSourceControlChangelistsWidget::OnSourceControlStateChanged()
+{
+	Refresh();
 }
 
 void SSourceControlChangelistsWidget::OnChangelistsStatusUpdated(const FSourceControlOperationRef& InOperation, ECommandResult::Type InType)
 {
-	ISourceControlModule& SourceControlModule = ISourceControlModule::Get();
-	ISourceControlProvider& SourceControlProvider = SourceControlModule.GetProvider();
+	if (ISourceControlModule::Get().IsEnabled())
+	{
+		ISourceControlProvider& SourceControlProvider = ISourceControlModule::Get().GetProvider();
 
-	TArray<FSourceControlChangelistStateRef> ChangelistsStates;
-	SourceControlProvider.GetState(SourceControlProvider.GetChangelists(EStateCacheUsage::Use), ChangelistsStates, EStateCacheUsage::Use);
+		TArray<FSourceControlChangelistRef> Changelists = SourceControlProvider.GetChangelists(EStateCacheUsage::Use);
 
-	ChangelistsNodes.Reset(ChangelistsStates.Num());
+		TArray<FSourceControlChangelistStateRef> ChangelistsStates;
+		SourceControlProvider.GetState(Changelists, ChangelistsStates, EStateCacheUsage::Use);
 
-	for (FSourceControlChangelistStateRef ChangelistState : ChangelistsStates)
-	{
-		FChangelistTreeItemRef ChangelistTreeItem = MakeShareable(new FChangelistTreeItem(ChangelistState));
+		ChangelistsNodes.Reset(ChangelistsStates.Num());
 
-		for (FSourceControlStateRef FileRef : ChangelistState->GetFilesStates())
+		for (FSourceControlChangelistStateRef ChangelistState : ChangelistsStates)
 		{
-			FChangelistTreeItemRef FileTreeItem = MakeShareable(new FFileTreeItem(FileRef));
-			ChangelistTreeItem->AddChild(FileTreeItem);
-		}
+			FChangelistTreeItemRef ChangelistTreeItem = MakeShareable(new FChangelistTreeItem(ChangelistState));
 
-		ChangelistsNodes.Add(ChangelistTreeItem);
+			for (FSourceControlStateRef FileRef : ChangelistState->GetFilesStates())
+			{
+				FChangelistTreeItemRef FileTreeItem = MakeShareable(new FFileTreeItem(FileRef));
+				ChangelistTreeItem->AddChild(FileTreeItem);
+			}
+
+			ChangelistsNodes.Add(ChangelistTreeItem);
+		}
+	}
+	else
+	{
+		ChangelistsNodes.Empty();
 	}
 
 	TreeView->RequestTreeRefresh();
diff --git a/Engine/Source/Editor/SourceControlWindows/Private/SSourceControlChangelists.h b/Engine/Source/Editor/SourceControlWindows/Private/SSourceControlChangelists.h
--- a/Engine/Source/Editor/SourceControlWindows/Private/SSourceControlChangelists.h
+++ b/Engine/Source/Editor/SourceControlWindows/Private/SSourceControlChangelists.h
@@ -93,6 +93,9 @@ private:
 	void Refresh();
 	void OnChangelistsStatusUpdated(const FSourceControlOperationRef&, ECommandResult::Type);
 
+	void OnSourceControlStateChanged();
+	void OnSourceControlProviderChanged(ISourceControlProvider& OldProvider, ISourceControlProvider& NewProvider);
+
 private:
 
 	/** Changelists (root nodes) */
@@ -102,4 +105,7 @@ private:
 	TSharedPtr<SChangelistTree> TreeView;
 
 	TSharedPtr<class FUpdatePendingChangelistsStatus, ESPMode::ThreadSafe> UpdatePendingChangelistsOperation;
+
+	/** Source control state changed delegate handle */
+	FDelegateHandle SourceControlStateChangedDelegateHandle;
 };
