Changelists window - Place the "Changelist" window under the "Tools -> Source Control" group

diff --git a/Engine/Source/Editor/SourceControlWindows/Private/SourceControlWindowsModule.cpp b/Engine/Source/Editor/SourceControlWindows/Private/SourceControlWindowsModule.cpp
--- a/Engine/Source/Editor/SourceControlWindows/Private/SourceControlWindowsModule.cpp
+++ b/Engine/Source/Editor/SourceControlWindows/Private/SourceControlWindowsModule.cpp
@@ -52,16 +52,29 @@ void FSourceControlWindowsModule::StartupModule()
 {
 	ISourceControlWindowsModule::StartupModule();
 
+	// We're going to call a static function in the editor style module, so we need to make sure the module has actually been loaded
+	FModuleManager::Get().LoadModuleChecked("EditorStyle");
+
+	// Create a Source Control group under the Tools category
+	const FSlateIcon SourceControlIcon(FEditorStyle::GetStyleSetName(), "SourceControl.ChangelistsTab");
+	const IWorkspaceMenuStructure& MenuStructure = WorkspaceMenu::GetMenuStructure();
+	TSharedRef<FWorkspaceItem> SourceControlGroup = MenuStructure.GetToolsCategory()->AddGroup(
+		LOCTEXT("WorkspaceMenu_SourceControlCategory", "Source Control"),
+		LOCTEXT("SourceControlMenuTooltipText", "Source Control Operations"),
+		SourceControlIcon,
+		true);
+
+	// Register the changlist tab spawner
 	FGlobalTabmanager::Get()->RegisterNomadTabSpawner(SourceControlChangelistsTabName, FOnSpawnTab::CreateRaw(this, &FSourceControlWindowsModule::CreateChangelistsTab))
-		.SetDisplayName(LOCTEXT("ChangelistsTabTitle", "Changelists"));
+		.SetDisplayName(LOCTEXT("ChangelistsTabTitle", "Changelists"))
+		.SetGroup(SourceControlGroup)
+		.SetIcon(SourceControlIcon);
 
 #if WITH_HOT_RELOAD
 	// This code attempts to relaunch the GameplayCueEditor tab when you hotreload this module
 	if (GIsHotReload && FSlateApplication::IsInitialized())
 	{
-		FLevelEditorModule& LevelEditorModule = FModuleManager::GetModuleChecked<FLevelEditorModule>(TEXT("LevelEditor"));
-		TSharedPtr<FTabManager> LevelEditorTabManager = LevelEditorModule.GetLevelEditorTabManager();
-		LevelEditorTabManager->TryInvokeTab(SourceControlChangelistsTabName);
+		ShowChangelistsTab();
 	}
 #endif // WITH_HOT_RELOAD
 }
