Changelists window - Improved window refresh mechanism

diff --git a/Engine/Source/Editor/SourceControlWindows/Private/SSourceControlChangelists.cpp b/Engine/Source/Editor/SourceControlWindows/Private/SSourceControlChangelists.cpp
--- a/Engine/Source/Editor/SourceControlWindows/Private/SSourceControlChangelists.cpp
+++ b/Engine/Source/Editor/SourceControlWindows/Private/SSourceControlChangelists.cpp
@@ -68,50 +68,63 @@ void SSourceControlChangelistsWidget::Construct(const FArguments& InArgs)
 
 	TreeView = CreateTreeviewWidget();
 
-	Refresh();
-
 	ChildSlot
+	[
+		SNew(SScrollBorder, TreeView.ToSharedRef())
+		.Visibility(TAttribute<EVisibility>::Create(TAttribute<EVisibility>::FGetter::CreateLambda([]()->EVisibility { return ISourceControlModule::Get().IsEnabled() ? EVisibility::Visible : EVisibility::Hidden; })))
 		[
-			SNew(SScrollBorder, TreeView.ToSharedRef())
-			[
-				TreeView.ToSharedRef()
-			]
-		];
+			TreeView.ToSharedRef()
+		]
+	];
+
+	bShouldRefresh = true;
 }
 
-void SSourceControlChangelistsWidget::Refresh()
+void SSourceControlChangelistsWidget::Tick(const FGeometry& AllottedGeometry, const double InCurrentTime, const float InDeltaTime)
+{
+	if (bShouldRefresh)
+	{
+		if (ISourceControlModule::Get().IsEnabled())
+		{
+			RequestRefresh();
+			bShouldRefresh = false;
+		}
+		else
+		{
+			// No provider available, clear changelist tree
+			ClearChangelistsTree();
+		}
+	}
+}
+
+void SSourceControlChangelistsWidget::RequestRefresh()
 {
 	if (ISourceControlModule::Get().IsEnabled())
 	{
-		UpdatePendingChangelistsOperation = ISourceControlOperation::Create<FUpdatePendingChangelistsStatus>();
+		TSharedRef<FUpdatePendingChangelistsStatus, ESPMode::ThreadSafe> UpdatePendingChangelistsOperation = ISourceControlOperation::Create<FUpdatePendingChangelistsStatus>();
 		UpdatePendingChangelistsOperation->SetUpdateAllChangelists(true);
 		UpdatePendingChangelistsOperation->SetUpdateFilesStates(true);
 
 		ISourceControlProvider& SourceControlProvider = ISourceControlModule::Get().GetProvider();
-		SourceControlProvider.Execute(UpdatePendingChangelistsOperation.ToSharedRef(), EConcurrency::Asynchronous, FSourceControlOperationComplete::CreateSP(this, &SSourceControlChangelistsWidget::OnChangelistsStatusUpdated));
+		SourceControlProvider.Execute(UpdatePendingChangelistsOperation, EConcurrency::Asynchronous, FSourceControlOperationComplete::CreateSP(this, &SSourceControlChangelistsWidget::OnChangelistsStatusUpdated));
 	}
 	else
 	{
 		// No provider available, clear changelist tree
-		ChangelistsNodes.Empty();
-		TreeView->RequestTreeRefresh();
+		ClearChangelistsTree();
 	}
 }
 
-void SSourceControlChangelistsWidget::OnSourceControlProviderChanged(ISourceControlProvider& OldProvider, ISourceControlProvider& NewProvider)
+void SSourceControlChangelistsWidget::ClearChangelistsTree()
 {
-	OldProvider.UnregisterSourceControlStateChanged_Handle(SourceControlStateChangedDelegateHandle);
-	SourceControlStateChangedDelegateHandle = NewProvider.RegisterSourceControlStateChanged_Handle(FSourceControlStateChanged::FDelegate::CreateSP(this, &SSourceControlChangelistsWidget::OnSourceControlStateChanged));
-
-	OnSourceControlStateChanged();
-}
-
-void SSourceControlChangelistsWidget::OnSourceControlStateChanged()
-{
-	Refresh();
+	if (!ChangelistsNodes.IsEmpty())
+	{
+		ChangelistsNodes.Empty();
+		TreeView->RequestTreeRefresh();
+	}
 }
 
-void SSourceControlChangelistsWidget::OnChangelistsStatusUpdated(const FSourceControlOperationRef& InOperation, ECommandResult::Type InType)
+void SSourceControlChangelistsWidget::Refresh()
 {
 	if (ISourceControlModule::Get().IsEnabled())
 	{
@@ -136,31 +149,31 @@ void SSourceControlChangelistsWidget::OnChangelistsStatusUpdated(const FSourceCo
 
 			ChangelistsNodes.Add(ChangelistTreeItem);
 		}
+
+		TreeView->RequestTreeRefresh();
 	}
 	else
 	{
-		ChangelistsNodes.Empty();
+		ClearChangelistsTree();
 	}
-
-	TreeView->RequestTreeRefresh();
 }
 
-bool SSourceControlChangelistsWidget::OnIsSelectableOrNavigable(FChangelistTreeItemPtr InItem) const
+void SSourceControlChangelistsWidget::OnSourceControlProviderChanged(ISourceControlProvider& OldProvider, ISourceControlProvider& NewProvider)
 {
-	if (!InItem)
-	{
-		return false;
-	}
+	OldProvider.UnregisterSourceControlStateChanged_Handle(SourceControlStateChangedDelegateHandle);
+	SourceControlStateChangedDelegateHandle = NewProvider.RegisterSourceControlStateChanged_Handle(FSourceControlStateChanged::FDelegate::CreateSP(this, &SSourceControlChangelistsWidget::OnSourceControlStateChanged));
 
-	TArray<FChangelistTreeItemPtr> SelectedItems = TreeView->GetSelectedItems();
+	bShouldRefresh = true;
+}
 
-	// Prevent selecting changelists and files at the same time.
-	if (!SelectedItems.IsEmpty() && SelectedItems[0]->GetTreeItemType() != InItem->GetTreeItemType())
-	{
-		return false;
-	}
+void SSourceControlChangelistsWidget::OnSourceControlStateChanged()
+{
+	Refresh();
+}
 
-	return true;
+void SSourceControlChangelistsWidget::OnChangelistsStatusUpdated(const FSourceControlOperationRef& InOperation, ECommandResult::Type InType)
+{
+	Refresh();
 }
 
 TSharedRef<SChangelistTree> SSourceControlChangelistsWidget::CreateTreeviewWidget()
@@ -171,7 +184,6 @@ TSharedRef<SChangelistTree> SSourceControlChangelistsWidget::CreateTreeviewWidge
 		.OnGenerateRow(this, &SSourceControlChangelistsWidget::OnGenerateRow)
 		.OnGetChildren(this, &SSourceControlChangelistsWidget::OnGetChildren)
 		.SelectionMode(ESelectionMode::Multi)
-		.OnIsSelectableOrNavigable(this, &SSourceControlChangelistsWidget::OnIsSelectableOrNavigable)
 		.HeaderRow
 		(
 			SNew(SHeaderRow)
diff --git a/Engine/Source/Editor/SourceControlWindows/Private/SSourceControlChangelists.h b/Engine/Source/Editor/SourceControlWindows/Private/SSourceControlChangelists.h
--- a/Engine/Source/Editor/SourceControlWindows/Private/SSourceControlChangelists.h
+++ b/Engine/Source/Editor/SourceControlWindows/Private/SSourceControlChangelists.h
@@ -88,13 +88,15 @@ private:
 	TSharedRef<ITableRow> OnGenerateRow(FChangelistTreeItemPtr InTreeItem, const TSharedRef<STableViewBase>& OwnerTable);
 	void OnGetChildren(FChangelistTreeItemPtr InParent, TArray<FChangelistTreeItemPtr>& OutChildren);
 
-	bool OnIsSelectableOrNavigable(FChangelistTreeItemPtr InItem) const;
-
+	void RequestRefresh();
 	void Refresh();
-	void OnChangelistsStatusUpdated(const FSourceControlOperationRef&, ECommandResult::Type);
+	void ClearChangelistsTree();
 
 	void OnSourceControlStateChanged();
 	void OnSourceControlProviderChanged(ISourceControlProvider& OldProvider, ISourceControlProvider& NewProvider);
+	void OnChangelistsStatusUpdated(const FSourceControlOperationRef& InOperation, ECommandResult::Type InType);
+
+	virtual void Tick(const FGeometry& AllottedGeometry, const double InCurrentTime, const float InDeltaTime) override;
 
 private:
 
@@ -104,8 +106,8 @@ private:
 	/** Changelist treeview widget */
 	TSharedPtr<SChangelistTree> TreeView;
 
-	TSharedPtr<class FUpdatePendingChangelistsStatus, ESPMode::ThreadSafe> UpdatePendingChangelistsOperation;
-
 	/** Source control state changed delegate handle */
 	FDelegateHandle SourceControlStateChangedDelegateHandle;
+
+	bool bShouldRefresh;
 };
