Changelist window - Selection improvement
* Prevent selection of files & changelists at the same time
* Only allow selection of files from the same changelist

diff --git a/Engine/Source/Editor/SourceControlWindows/Private/SSourceControlChangelists.cpp b/Engine/Source/Editor/SourceControlWindows/Private/SSourceControlChangelists.cpp
--- a/Engine/Source/Editor/SourceControlWindows/Private/SSourceControlChangelists.cpp
+++ b/Engine/Source/Editor/SourceControlWindows/Private/SSourceControlChangelists.cpp
@@ -176,6 +176,31 @@ void SSourceControlChangelistsWidget::OnChangelistsStatusUpdated(const FSourceCo
 	Refresh();
 }
 
+void SChangelistTree::Private_SetItemSelection(FChangelistTreeItemPtr TheItem, bool bShouldBeSelected, bool bWasUserDirected)
+{
+	bool bAllowSelectionChange = true;
+
+	if (bShouldBeSelected && !SelectedItems.IsEmpty())
+	{
+		// Prevent selecting changelists and files at the same time.
+		FChangelistTreeItemPtr CurrentlySelectedItem = (*SelectedItems.begin());
+		if (TheItem->GetTreeItemType() != CurrentlySelectedItem->GetTreeItemType())
+		{
+			bAllowSelectionChange = false;
+		}
+		// Prevent selecting items that don't share the same root
+		else if (TheItem->GetParent() != CurrentlySelectedItem->GetParent())
+		{
+			bAllowSelectionChange = false;
+		}
+	}
+
+	if (bAllowSelectionChange)
+	{
+		STreeView::Private_SetItemSelection(TheItem, bShouldBeSelected, bWasUserDirected);
+	}
+}
+
 TSharedRef<SChangelistTree> SSourceControlChangelistsWidget::CreateTreeviewWidget()
 {
 	return SAssignNew(TreeView, SChangelistTree)
diff --git a/Engine/Source/Editor/SourceControlWindows/Private/SSourceControlChangelists.h b/Engine/Source/Editor/SourceControlWindows/Private/SSourceControlChangelists.h
--- a/Engine/Source/Editor/SourceControlWindows/Private/SSourceControlChangelists.h
+++ b/Engine/Source/Editor/SourceControlWindows/Private/SSourceControlChangelists.h
@@ -66,7 +66,11 @@ protected:
 	TreeItemType Type;
 };
 
-typedef STreeView<FChangelistTreeItemPtr> SChangelistTree;
+class SChangelistTree : public STreeView<FChangelistTreeItemPtr>
+{
+private:
+	virtual void Private_SetItemSelection(FChangelistTreeItemPtr TheItem, bool bShouldBeSelected, bool bWasUserDirected = false) override;
+};
 
 class SSourceControlChangelistsWidget : public SCompoundWidget
 {
