Added contextual menu options for file actions (diff agains depot, show history, locate file)
Better cleanup in perforce revert operation
Made revert work on changelists
Fixed issue where no changelists would be displayed if the default changelist was empty.

diff --git a/Engine/Plugins/Developer/PerforceSourceControl/Source/PerforceSourceControl/Private/PerforceSourceControlChangelist.cpp b/Engine/Plugins/Developer/PerforceSourceControl/Source/PerforceSourceControl/Private/PerforceSourceControlChangelist.cpp
--- a/Engine/Plugins/Developer/PerforceSourceControl/Source/PerforceSourceControl/Private/PerforceSourceControlChangelist.cpp
+++ b/Engine/Plugins/Developer/PerforceSourceControl/Source/PerforceSourceControl/Private/PerforceSourceControlChangelist.cpp
@@ -7,5 +7,6 @@ const FPerforceSourceControlChangelist FPerforceSourceControlChangelist::Default
 
 FPerforceSourceControlChangelist::FPerforceSourceControlChangelist()
 	: ChangelistNumber(DEFAULT_CHANGELIST_NUMBER)
+	, bInitialized(false)
 {
 }
\ No newline at end of file
diff --git a/Engine/Plugins/Developer/PerforceSourceControl/Source/PerforceSourceControl/Private/PerforceSourceControlChangelist.h b/Engine/Plugins/Developer/PerforceSourceControl/Source/PerforceSourceControl/Private/PerforceSourceControlChangelist.h
--- a/Engine/Plugins/Developer/PerforceSourceControl/Source/PerforceSourceControl/Private/PerforceSourceControlChangelist.h
+++ b/Engine/Plugins/Developer/PerforceSourceControl/Source/PerforceSourceControl/Private/PerforceSourceControlChangelist.h
@@ -12,6 +12,7 @@ public:
 
 	explicit FPerforceSourceControlChangelist(int32 InChangelistNumber)
 		: ChangelistNumber(InChangelistNumber)
+		, bInitialized(true)
 	{
 	}
 
@@ -32,7 +33,18 @@ public:
 
 	bool IsDefault() const
 	{
-		return *this == DefaultChangelist;
+		return ChangelistNumber == DefaultChangelist.ChangelistNumber;
+	}
+
+	bool IsInitialized() const
+	{
+		return bInitialized;
+	}
+
+	void Reset()
+	{
+		ChangelistNumber = DefaultChangelist.ChangelistNumber;
+		bInitialized = false;
 	}
 
 	friend FORCEINLINE uint32 GetTypeHash(const FPerforceSourceControlChangelist& PerforceChangelist)
@@ -48,6 +60,7 @@ public:
 public:
 	static const FPerforceSourceControlChangelist DefaultChangelist;
 	int32 ChangelistNumber;
+	bool bInitialized;
 };
 
 typedef TSharedRef<class FPerforceSourceControlChangelist, ESPMode::ThreadSafe> FPerforceSourceControlChangelistRef;
\ No newline at end of file
diff --git a/Engine/Plugins/Developer/PerforceSourceControl/Source/PerforceSourceControl/Private/PerforceSourceControlCommand.cpp b/Engine/Plugins/Developer/PerforceSourceControl/Source/PerforceSourceControl/Private/PerforceSourceControlCommand.cpp
--- a/Engine/Plugins/Developer/PerforceSourceControl/Source/PerforceSourceControl/Private/PerforceSourceControlCommand.cpp
+++ b/Engine/Plugins/Developer/PerforceSourceControl/Source/PerforceSourceControl/Private/PerforceSourceControlCommand.cpp
@@ -17,7 +17,6 @@ FPerforceSourceControlCommand::FPerforceSourceControlCommand(const TSharedRef<cl
 	, bConnectionDropped(false)
 	, bAutoDelete(true)
 	, Concurrency(EConcurrency::Synchronous)
-	, Changelist(FPerforceSourceControlChangelist::DefaultChangelist)
 {
 	// grab the providers settings here, so we don't access them once the worker thread is launched
 	check(IsInGameThread());
diff --git a/Engine/Plugins/Developer/PerforceSourceControl/Source/PerforceSourceControl/Private/PerforceSourceControlOperations.cpp b/Engine/Plugins/Developer/PerforceSourceControl/Source/PerforceSourceControl/Private/PerforceSourceControlOperations.cpp
--- a/Engine/Plugins/Developer/PerforceSourceControl/Source/PerforceSourceControl/Private/PerforceSourceControlOperations.cpp
+++ b/Engine/Plugins/Developer/PerforceSourceControl/Source/PerforceSourceControl/Private/PerforceSourceControlOperations.cpp
@@ -583,19 +583,55 @@ bool FPerforceRevertWorker::Execute(FPerforceSourceControlCommand& InCommand)
 		FPerforceConnection& Connection = ScopedConnection.GetConnection();
 		TArray<FString> Parameters;
 
-		AppendChangelistParameter(Parameters);
-		Parameters.Append(InCommand.Files);
+		if (InCommand.Changelist.IsInitialized())
+		{
+			Parameters.Add(TEXT("-c"));
+			Parameters.Add(InCommand.Changelist.ToString());
+		}
+		else
+		{
+			AppendChangelistParameter(Parameters);
+		}
+
+		if (InCommand.Files.Num() != 0)
+		{
+			Parameters.Append(InCommand.Files);
+		}
+		else if(InCommand.Changelist.IsInitialized()) // Note: safety net here, as we probably never want to revert everything
+		{
+			Parameters.Add(TEXT("//..."));
+		}
 
 		FP4RecordSet Records;
 		InCommand.bCommandSuccessful = Connection.RunCommand(TEXT("revert"), Parameters, Records, InCommand.ResultInfo.ErrorMessages, FOnIsCancelled::CreateRaw(&InCommand, &FPerforceSourceControlCommand::IsCanceled), InCommand.bConnectionDropped);
 		ParseRecordSetForState(Records, OutResults);
+		ChangelistToUpdate = InCommand.Changelist;
 	}
 	return InCommand.bCommandSuccessful;
 }
 
+bool RemoveFilesFromChangelist(const TMap<FString, EPerforceState::Type>& Results, TSharedRef<FPerforceSourceControlChangelistState, ESPMode::ThreadSafe>& ChangelistState)
+{
+	return ChangelistState->Files.RemoveAll([&Results](FSourceControlStateRef& State) -> bool
+		{
+			return Algo::AnyOf(Results, [&State](auto& Result) {
+				return State->GetFilename() == Result.Key;
+				});
+		}) > 0;
+}
+
+bool RemoveFilesFromChangelist(const TMap<FString, EPerforceState::Type>& Results, const FPerforceSourceControlChangelist& Changelist)
+{
+	FPerforceSourceControlModule& PerforceSourceControl = FPerforceSourceControlModule::Get();
+	TSharedRef<FPerforceSourceControlChangelistState, ESPMode::ThreadSafe> ChangelistState = PerforceSourceControl.GetProvider().GetStateInternal(Changelist);
+	return RemoveFilesFromChangelist(Results, ChangelistState);
+}
+
 bool FPerforceRevertWorker::UpdateStates() const
 {
-	return UpdateCachedStates(OutResults);
+	bool bUpdatedCachedStates = UpdateCachedStates(OutResults);
+	bool bUpdatedChangelists = ChangelistToUpdate.IsInitialized() && RemoveFilesFromChangelist(OutResults, ChangelistToUpdate);
+	return bUpdatedCachedStates || bUpdatedChangelists;
 }
 
 FName FPerforceSyncWorker::GetName() const
@@ -1367,7 +1403,8 @@ bool FPerforceUpdateStatusWorker::Execute(FPerforceSourceControlCommand& InComma
 			TArray<FString> Parameters = InCommand.Files;
 			Parameters.Add(FileQuery);
 			FP4RecordSet Records;
-			InCommand.bCommandSuccessful &= Connection.RunCommand(TEXT("opened"), Parameters, Records, InCommand.ResultInfo.ErrorMessages, FOnIsCancelled::CreateRaw(&InCommand, &FPerforceSourceControlCommand::IsCanceled), InCommand.bConnectionDropped);
+			Connection.RunCommand(TEXT("opened"), Parameters, Records, InCommand.ResultInfo.ErrorMessages, FOnIsCancelled::CreateRaw(&InCommand, &FPerforceSourceControlCommand::IsCanceled), InCommand.bConnectionDropped);
+			InCommand.bCommandSuccessful &= (InCommand.ResultInfo.ErrorMessages.Num() == 0);
 			ParseOpenedResults(Records, ANSI_TO_TCHAR(Connection.P4Client.GetClient().Text()), Connection.ClientRoot, OutStateMap);
 			RemoveRedundantErrors(InCommand, TEXT(" - no such file(s)."));
 			RemoveRedundantErrors(InCommand, TEXT("' is not under client's root '"));
@@ -1533,7 +1570,8 @@ bool FPerforceGetPendingChangelistsWorker::Execute(FPerforceSourceControlCommand
 				FP4RecordSet Records;
 				TMap<FString, EPerforceState::Type>& OutStateMap = OutCLFilesStateMap.Emplace_GetRef();
 
-				InCommand.bCommandSuccessful &= Connection.RunCommand(TEXT("opened"), Parameters, Records, InCommand.ResultInfo.ErrorMessages, FOnIsCancelled::CreateRaw(&InCommand, &FPerforceSourceControlCommand::IsCanceled), InCommand.bConnectionDropped);
+				Connection.RunCommand(TEXT("opened"), Parameters, Records, InCommand.ResultInfo.ErrorMessages, FOnIsCancelled::CreateRaw(&InCommand, &FPerforceSourceControlCommand::IsCanceled), InCommand.bConnectionDropped);
+				InCommand.bCommandSuccessful &= (InCommand.ResultInfo.ErrorMessages.Num() == 0);
 				ParseOpenedResults(Records, ANSI_TO_TCHAR(Connection.P4Client.GetClient().Text()), Connection.ClientRoot, OutStateMap);
 			}
 		}
@@ -1908,21 +1946,7 @@ bool FPerforceRevertUnchangedWorker::Execute(class FPerforceSourceControlCommand
 bool FPerforceRevertUnchangedWorker::UpdateStates() const
 {
 	bool bUpdatedStates = UpdateCachedStates(OutResults);
-	bool bUpdatedChangelistState = false;
-
-	if (OutResults.Num() > 0)
-	{
-		FPerforceSourceControlModule& PerforceSourceControl = FPerforceSourceControlModule::Get();
-		TSharedRef<FPerforceSourceControlChangelistState, ESPMode::ThreadSafe> ChangelistState = PerforceSourceControl.GetProvider().GetStateInternal(ChangelistToUpdate);
-
-		bUpdatedChangelistState = ChangelistState->Files.RemoveAll([this](FSourceControlStateRef& State) -> bool
-			{
-				return Algo::AnyOf(OutResults, [&State](auto& OutResult) {
-					return State->GetFilename() == OutResult.Key;
-					});
-			}) > 0;
-	}
-
+	bool bUpdatedChangelistState = ChangelistToUpdate.IsInitialized() && RemoveFilesFromChangelist(OutResults, ChangelistToUpdate);
 	return bUpdatedStates || bUpdatedChangelistState;
 }
 
diff --git a/Engine/Plugins/Developer/PerforceSourceControl/Source/PerforceSourceControl/Private/PerforceSourceControlOperations.h b/Engine/Plugins/Developer/PerforceSourceControl/Source/PerforceSourceControl/Private/PerforceSourceControlOperations.h
--- a/Engine/Plugins/Developer/PerforceSourceControl/Source/PerforceSourceControl/Private/PerforceSourceControlOperations.h
+++ b/Engine/Plugins/Developer/PerforceSourceControl/Source/PerforceSourceControl/Private/PerforceSourceControlOperations.h
@@ -90,6 +90,9 @@ public:
 public:
 	/** Map of filenames to perforce state */
 	TMap<FString, EPerforceState::Type> OutResults;
+
+	/** Changelist to be udpated */
+	FPerforceSourceControlChangelist ChangelistToUpdate;
 };
 
 class FPerforceSyncWorker : public IPerforceSourceControlWorker
diff --git a/Engine/Plugins/Developer/PerforceSourceControl/Source/PerforceSourceControl/Private/PerforceSourceControlProvider.cpp b/Engine/Plugins/Developer/PerforceSourceControl/Source/PerforceSourceControl/Private/PerforceSourceControlProvider.cpp
--- a/Engine/Plugins/Developer/PerforceSourceControl/Source/PerforceSourceControl/Private/PerforceSourceControlProvider.cpp
+++ b/Engine/Plugins/Developer/PerforceSourceControl/Source/PerforceSourceControl/Private/PerforceSourceControlProvider.cpp
@@ -359,7 +359,7 @@ ECommandResult::Type FPerforceSourceControlProvider::Execute( const FSourceContr
 		return ECommandResult::Failed;
 	}
 
-	FPerforceSourceControlChangelist Changelist = InChangelist ? InChangelist.ToSharedRef().Get() : FPerforceSourceControlChangelist::DefaultChangelist;
+	FPerforceSourceControlChangelist Changelist = InChangelist ? InChangelist.ToSharedRef().Get() : FPerforceSourceControlChangelist();
 
 	// fire off operation
 	if(InConcurrency == EConcurrency::Synchronous)
diff --git a/Engine/Source/Editor/SourceControlWindows/Private/SSourceControlChangelists.cpp b/Engine/Source/Editor/SourceControlWindows/Private/SSourceControlChangelists.cpp
--- a/Engine/Source/Editor/SourceControlWindows/Private/SSourceControlChangelists.cpp
+++ b/Engine/Source/Editor/SourceControlWindows/Private/SSourceControlChangelists.cpp
@@ -17,6 +17,11 @@
 #include "ToolMenus.h"
 #include "Widgets/Images/SLayeredImage.h"
 #include "SSourceControlDescription.h"
+#include "SourceControlWindows.h"
+#include "AssetToolsModule.h"
+#include "ContentBrowserModule.h"
+#include "IContentBrowserSingleton.h"
+
 
 #define LOCTEXT_NAMESPACE "SourceControlChangelist"
 
@@ -414,6 +419,26 @@ void SSourceControlChangelistsWidget::OnRevertUnchanged()
 	Refresh();
 }
 
+void SSourceControlChangelistsWidget::OnRevert()
+{
+	ISourceControlProvider& SourceControlProvider = ISourceControlModule::Get().GetProvider();
+
+	auto RevertOperation = ISourceControlOperation::Create<FRevert>();
+
+	if (GetCurrentChangelist() != nullptr)
+	{
+		// Operation on full changelist
+		SourceControlProvider.Execute(RevertOperation, GetCurrentChangelist());
+	}
+	else
+	{
+		// Operation on files in a changelist
+		SourceControlProvider.Execute(RevertOperation, GetChangelistFromFileSelection(), GetSelectedFiles());
+	}
+
+	Refresh();
+}
+
 void SSourceControlChangelistsWidget::OnSubmitChangelist()
 {
 	FSourceControlChangelistPtr Changelist = GetCurrentChangelist();
@@ -436,6 +461,69 @@ bool SSourceControlChangelistsWidget::CanSubmitChangelist()
 	return Changelist != nullptr && Changelist->GetFilesStates().Num() >= 0 && Changelist->GetShelvedFilesStates().Num() == 0;
 }
 
+void SSourceControlChangelistsWidget::OnLocateFile()
+{ 
+	FAssetToolsModule& AssetToolsModule = FModuleManager::LoadModuleChecked<FAssetToolsModule>(TEXT("AssetTools"));
+
+	TArray<UObject*> AssetsToSync;
+
+	TArray<FString> SelectedFiles = GetSelectedFiles();
+	AssetsToSync.Reserve(SelectedFiles.Num());
+
+	for (const FString& SelectedFile : SelectedFiles)
+	{
+		FString AssetPackageName;
+
+		if (FPackageName::TryConvertFilenameToLongPackageName(SelectedFile, AssetPackageName))
+		{
+			UPackage* AssetPackage = LoadPackage(nullptr, *AssetPackageName, LOAD_ForDiff | LOAD_DisableCompileOnLoad);
+
+			// grab the asset from the package - we assume asset name matches file name
+			FString AssetName = FPaths::GetBaseFilename(SelectedFile);
+			UObject* SelectedAsset = FindObject<UObject>(AssetPackage, *AssetName);
+
+			if (SelectedAsset)
+			{
+				AssetsToSync.Add(SelectedAsset);
+			}
+		}
+	}
+
+	if (AssetsToSync.Num() > 0)
+	{
+		FContentBrowserModule& ContentBrowserModule = FModuleManager::Get().LoadModuleChecked<FContentBrowserModule>("ContentBrowser");
+		ContentBrowserModule.Get().SyncBrowserToAssets(AssetsToSync, true);
+	}
+}
+
+bool SSourceControlChangelistsWidget::CanLocateFile()
+{
+	return GetSelectedFiles().Num() > 0;
+}
+
+void SSourceControlChangelistsWidget::OnShowHistory()
+{
+	TArray<FString> SelectedFiles = GetSelectedFiles();
+	if (SelectedFiles.Num() > 0)
+	{
+		FSourceControlWindows::DisplayRevisionHistory(SelectedFiles);
+	}
+}
+
+void SSourceControlChangelistsWidget::OnDiffAgainstDepot()
+{
+	TArray<FString> SelectedFiles = GetSelectedFiles();
+	if (SelectedFiles.Num() > 0)
+	{
+		FSourceControlWindows::DiffAgainstWorkspace(SelectedFiles[0]);
+	} 
+}
+
+bool SSourceControlChangelistsWidget::CanDiffAgainstDepot()
+{
+	return GetSelectedFiles().Num() == 1;
+}
+
 TSharedPtr<SWidget> SSourceControlChangelistsWidget::OnOpenContextMenu()
 {
 	UToolMenus* ToolMenus = UToolMenus::Get();
@@ -454,6 +542,30 @@ TSharedPtr<SWidget> SSourceControlChangelistsWidget::OnOpenContextMenu()
 
 	FToolMenuSection& Section = Menu->AddSection("Source Control");
 	
+	// This should appear only on change lists
+	if (bHasSelectedChangelist)
+	{
+		Section.AddMenuEntry("SubmitChangelist", LOCTEXT("SourceControl_SubmitChangelist", "Submit Changelist"), LOCTEXT("SourceControl_SubmitChangeslit_Tooltip", "Submits a changelist"), FSlateIcon(),
+			FUIAction(
+				FExecuteAction::CreateSP(this, &SSourceControlChangelistsWidget::OnSubmitChangelist),
+				FCanExecuteAction::CreateSP(this, &SSourceControlChangelistsWidget::CanSubmitChangelist)));
+	}
+
+	// This can appear on both files & changelist
+	if (!bHasEmptySelection)
+	{
+		Section.AddMenuEntry("RevertUnchanged", LOCTEXT("SourceControl_RevertUnchanged", "Revert Unchanged"), LOCTEXT("SourceControl_Revert_Unchanged_Tooltip", "Reverts unchanged files & changelists"), FSlateIcon(),
+			FUIAction(FExecuteAction::CreateSP(this, &SSourceControlChangelistsWidget::OnRevertUnchanged)));
+
+		Section.AddMenuEntry("Revert", LOCTEXT("SourceControl_Revert", "Revert Files"), LOCTEXT("SourceControl_Revert_Tooltip", "Reverts all files in the changelist or from the selection"), FSlateIcon(),
+			FUIAction(FExecuteAction::CreateSP(this, &SSourceControlChangelistsWidget::OnRevert)));
+	}
+
+	if (bHasEmptySelection || bHasSelectedChangelist)
+	{
+		Section.AddSeparator("Changelists");
+	}
+
 	// This should appear only if we have no selection
 	if (bHasEmptySelection)
 	{
@@ -461,7 +573,6 @@ TSharedPtr<SWidget> SSourceControlChangelistsWidget::OnOpenContextMenu()
 			FUIAction(FExecuteAction::CreateSP(this, &SSourceControlChangelistsWidget::OnNewChangelist)));
 	}
 
-	// This should appear only on change lists
 	if (bHasSelectedChangelist)
 	{
 		Section.AddMenuEntry("EditChangelist", LOCTEXT("SourceControl_EditChangelist", "Edit Changelist"), LOCTEXT("SourceControl_Edit_Changelist_Tooltip", "Edit a changelist description"), FSlateIcon(),
@@ -471,18 +582,26 @@ TSharedPtr<SWidget> SSourceControlChangelistsWidget::OnOpenContextMenu()
 			FUIAction(
 				FExecuteAction::CreateSP(this, &SSourceControlChangelistsWidget::OnDeleteChangelist),
 				FCanExecuteAction::CreateSP(this, &SSourceControlChangelistsWidget::CanDeleteChangelist)));
-
-		Section.AddMenuEntry("SubmitChangelist", LOCTEXT("SourceControl_SubmitChangelist", "Submit Changelist"), LOCTEXT("SourceControl_SubmitChangeslit_Tooltip", "Submits a changelist"), FSlateIcon(),
-			FUIAction(
-				FExecuteAction::CreateSP(this, &SSourceControlChangelistsWidget::OnSubmitChangelist),
-				FCanExecuteAction::CreateSP(this, &SSourceControlChangelistsWidget::CanSubmitChangelist)));
 	}
 
-	// This can appear on both files & changelist
-	if (!bHasEmptySelection)
+	// Files-only operations
+	if (!bHasEmptySelection && !bHasSelectedChangelist)
 	{
-		Section.AddMenuEntry("RevertUnchanged", LOCTEXT("SourceControl_RevertUnchanged", "Revert Unchanged"), LOCTEXT("SourceControl_Revert_Unchanged_Tooltip", "Reverts unchanged files & changelists"), FSlateIcon(),
-			FUIAction(FExecuteAction::CreateSP(this, &SSourceControlChangelistsWidget::OnRevertUnchanged)));
+		Section.AddSeparator("Files");
+
+		Section.AddMenuEntry("Locate File", LOCTEXT("SourceControl_LocateFile", "Locate File"), LOCTEXT("SourceControl_LocateFile_Tooltip", "Locate File in Project..."), FSlateIcon(),
+			FUIAction(
+				FExecuteAction::CreateSP(this, &SSourceControlChangelistsWidget::OnLocateFile),
+				FCanExecuteAction::CreateSP(this, &SSourceControlChangelistsWidget::CanLocateFile)));
+
+		Section.AddMenuEntry("Show History", LOCTEXT("SourceControl_ShowHistory", "Show History"), LOCTEXT("SourceControl_ShowHistory_ToolTip", "Show File History From Selection..."), FSlateIcon(),
+			FUIAction(
+				FExecuteAction::CreateSP(this, &SSourceControlChangelistsWidget::OnShowHistory)));
+
+		Section.AddMenuEntry("Diff Against Local Version", LOCTEXT("SourceControl_DiffAgainstDepot", "Diff Against Depot"), LOCTEXT("SourceControl_DiffAgainstLocal_Tooltip", "Diff local file against depot revision."), FSlateIcon(),
+			FUIAction(
+				FExecuteAction::CreateSP(this, &SSourceControlChangelistsWidget::OnDiffAgainstDepot),
+				FCanExecuteAction::CreateSP(this, &SSourceControlChangelistsWidget::CanDiffAgainstDepot)));
 	}
 
 	return ToolMenus->GenerateWidget(Menu);
diff --git a/Engine/Source/Editor/SourceControlWindows/Private/SSourceControlChangelists.h b/Engine/Source/Editor/SourceControlWindows/Private/SSourceControlChangelists.h
--- a/Engine/Source/Editor/SourceControlWindows/Private/SSourceControlChangelists.h
+++ b/Engine/Source/Editor/SourceControlWindows/Private/SSourceControlChangelists.h
@@ -109,14 +109,25 @@ private:
 	/** Returns list of currently selected files */
 	TArray<FString> GetSelectedFiles();
 
+	/** Changelist operations */
 	void OnNewChangelist();
 	void OnDeleteChangelist();
 	bool CanDeleteChangelist();
 	void OnEditChangelist();
-	void OnRevertUnchanged();
 	void OnSubmitChangelist();
 	bool CanSubmitChangelist();
 
+	/** Changelist & File operations */
+	void OnRevertUnchanged();
+	void OnRevert();
+
+	/** Files operations */
+	void OnLocateFile();
+	bool CanLocateFile();
+	void OnShowHistory();
+	void OnDiffAgainstDepot();
+	bool CanDiffAgainstDepot();
+
 	void OnSourceControlStateChanged();
 	void OnSourceControlProviderChanged(ISourceControlProvider& OldProvider, ISourceControlProvider& NewProvider);
 	void OnChangelistsStatusUpdated(const FSourceControlOperationRef& InOperation, ECommandResult::Type InType);
diff --git a/Engine/Source/Editor/SourceControlWindows/Private/SSourceControlHistory.cpp b/Engine/Source/Editor/SourceControlWindows/Private/SSourceControlHistory.cpp
--- a/Engine/Source/Editor/SourceControlWindows/Private/SSourceControlHistory.cpp
+++ b/Engine/Source/Editor/SourceControlWindows/Private/SSourceControlHistory.cpp
@@ -1641,3 +1641,25 @@ void FSourceControlWindows::DisplayRevisionHistory( const TArray<FString>& InPac
 		}
 	}
 }
+
+void FSourceControlWindows::DiffAgainstWorkspace(const FString& InFileName)
+{
+	FAssetToolsModule& AssetToolsModule = FModuleManager::LoadModuleChecked<FAssetToolsModule>(TEXT("AssetTools"));
+
+	UObject* SelectedAsset = nullptr;
+	FString AssetPackageName;
+
+	if (FPackageName::TryConvertFilenameToLongPackageName(InFileName, AssetPackageName))
+	{
+		UPackage* AssetPackage = LoadPackage(nullptr, *AssetPackageName, LOAD_ForDiff | LOAD_DisableCompileOnLoad);
+
+		// grab the asset from the package - we assume asset name matches file name
+		FString AssetName = FPaths::GetBaseFilename(InFileName);
+		SelectedAsset = FindObject<UObject>(AssetPackage, *AssetName);
+	}
+
+	if (SelectedAsset)
+	{
+		AssetToolsModule.Get().DiffAgainstDepot(SelectedAsset, AssetPackageName, SelectedAsset->GetName());
+	}
+}
diff --git a/Engine/Source/Editor/SourceControlWindows/Public/SourceControlWindows.h b/Engine/Source/Editor/SourceControlWindows/Public/SourceControlWindows.h
--- a/Engine/Source/Editor/SourceControlWindows/Public/SourceControlWindows.h
+++ b/Engine/Source/Editor/SourceControlWindows/Public/SourceControlWindows.h
@@ -94,6 +94,13 @@ public:
 	 */
 	static bool PromptForRevert(const TArray<FString>& InPackageNames );
 
+	/**
+	 * Displays file diff against workspace version
+	 * 
+	 * @param InFileName Name of the file to diff
+	 */
+	static void DiffAgainstWorkspace(const FString& InFileName);
+
 protected:
 	/** Callback for ChoosePackagesToCheckIn(), continues to bring up UI once source control operations are complete */
 	static void ChoosePackagesToCheckInCallback(const FSourceControlOperationRef& InOperation, ECommandResult::Type InResult, FSourceControlWindowsOnCheckInComplete OnCompleteDelegate);
