Added confirmation on submit/revert changelists
Fixed an issue in the get pending changelists when there is only the default changelist present

diff --git a/Engine/Plugins/Developer/PerforceSourceControl/Source/PerforceSourceControl/Private/PerforceSourceControlOperations.cpp b/Engine/Plugins/Developer/PerforceSourceControl/Source/PerforceSourceControl/Private/PerforceSourceControlOperations.cpp
--- a/Engine/Plugins/Developer/PerforceSourceControl/Source/PerforceSourceControl/Private/PerforceSourceControlOperations.cpp
+++ b/Engine/Plugins/Developer/PerforceSourceControl/Source/PerforceSourceControl/Private/PerforceSourceControlOperations.cpp
@@ -1530,6 +1530,8 @@ bool FPerforceGetPendingChangelistsWorker::Execute(FPerforceSourceControlCommand
 
 		TSharedRef<FUpdatePendingChangelistsStatus, ESPMode::ThreadSafe> Operation = StaticCastSharedRef<FUpdatePendingChangelistsStatus>(InCommand.Operation);
 
+		InCommand.bCommandSuccessful = true;
+
 		if (Operation->ShouldUpdateAllChangelists())
 		{
 			// First, insert the default changelist which always exists
@@ -1544,7 +1546,8 @@ bool FPerforceGetPendingChangelistsWorker::Execute(FPerforceSourceControlCommand
 			Parameters.Add(InCommand.ConnectionInfo.Workspace);			// <workspace>
 
 			FP4RecordSet Records;
-			InCommand.bCommandSuccessful = Connection.RunCommand(TEXT("changes"), Parameters, Records, InCommand.ResultInfo.ErrorMessages, FOnIsCancelled::CreateRaw(&InCommand, &FPerforceSourceControlCommand::IsCanceled), InCommand.bConnectionDropped);
+			Connection.RunCommand(TEXT("changes"), Parameters, Records, InCommand.ResultInfo.ErrorMessages, FOnIsCancelled::CreateRaw(&InCommand, &FPerforceSourceControlCommand::IsCanceled), InCommand.bConnectionDropped);
+			InCommand.bCommandSuccessful &= (InCommand.ResultInfo.ErrorMessages.Num() == 0);
 
 			ParseChangelistsResults(Records, OutChangelistsStates);
 		}
diff --git a/Engine/Source/Editor/SourceControlWindows/Private/SSourceControlChangelists.cpp b/Engine/Source/Editor/SourceControlWindows/Private/SSourceControlChangelists.cpp
--- a/Engine/Source/Editor/SourceControlWindows/Private/SSourceControlChangelists.cpp
+++ b/Engine/Source/Editor/SourceControlWindows/Private/SSourceControlChangelists.cpp
@@ -21,6 +21,7 @@
 #include "AssetToolsModule.h"
 #include "ContentBrowserModule.h"
 #include "IContentBrowserSingleton.h"
+#include "Misc/MessageDialog.h"
 
 
 #define LOCTEXT_NAMESPACE "SourceControlChangelist"
@@ -106,7 +107,16 @@ struct FFileTreeItem : public IChangelistTreeItem
 	FText GetDisplayText() const
 	{
 		FString Filename = FileState->GetFilename();
-		return FText::FromString(Filename);
+		FString AssetPackageName;
+
+		if (FPackageName::TryConvertFilenameToLongPackageName(Filename, AssetPackageName))
+		{
+			return FText::FromString(AssetPackageName);
+		}
+		else
+		{
+			return FText::FromString(Filename);
+		}
 	}
 
 	FSourceControlStateRef FileState;
@@ -421,8 +431,30 @@ void SSourceControlChangelistsWidget::OnRevertUnchanged()
 
 void SSourceControlChangelistsWidget::OnRevert()
 {
-	ISourceControlProvider& SourceControlProvider = ISourceControlModule::Get().GetProvider();
+	FText DialogText;
+	FText DialogTitle;
+
+	const bool bApplyOnChangelist = (GetCurrentChangelist() != nullptr);
 
+	if (bApplyOnChangelist)
+	{
+		DialogText = LOCTEXT("SourceControl_ConfirmRevertChangelist", "Are you sure you want to revert this changelist?");
+		DialogTitle = LOCTEXT("SourceControl_ConfirmRevertChangelist_Title", "Confirm changelist revert");
+	}
+	else
+	{
+		DialogText = LOCTEXT("SourceControl_ConfirmRevertFiles", "Are you sure you want to revert the selected files?");
+		DialogTitle = LOCTEXT("SourceControl_ConfirmReverFiles_Title", "Confirm files revert");
+	}
+	
+	EAppReturnType::Type UserConfirmation = FMessageDialog::Open(EAppMsgType::OkCancel, EAppReturnType::Ok, DialogText, &DialogTitle);
+
+	if (UserConfirmation != EAppReturnType::Ok)
+	{
+		return;
+	}
+
+	ISourceControlProvider& SourceControlProvider = ISourceControlModule::Get().GetProvider();
 	auto RevertOperation = ISourceControlOperation::Create<FRevert>();
 
 	if (GetCurrentChangelist() != nullptr)
@@ -448,11 +480,19 @@ void SSourceControlChangelistsWidget::OnSubmitChangelist()
 		return;
 	}
 
-	ISourceControlProvider& SourceControlProvider = ISourceControlModule::Get().GetProvider();
-	auto SubmitChangelistOperation = ISourceControlOperation::Create<FCheckIn>();
+	const FText DialogText = LOCTEXT("SourceControl_ConfirmSubmit", "Are you sure you want to submit this changelist?");
+	const FText DialogTitle = LOCTEXT("SourceControl_ConfirmSubmit_Title", "Confirm changelist submit");
+
+	EAppReturnType::Type UserConfirmation = FMessageDialog::Open(EAppMsgType::OkCancel, EAppReturnType::Ok, DialogText, & DialogTitle);
+
+	if (UserConfirmation == EAppReturnType::Ok)
+	{
+		ISourceControlProvider& SourceControlProvider = ISourceControlModule::Get().GetProvider();
+		auto SubmitChangelistOperation = ISourceControlOperation::Create<FCheckIn>();
 
-	SourceControlProvider.Execute(SubmitChangelistOperation, Changelist);
-	Refresh();	
+		SourceControlProvider.Execute(SubmitChangelistOperation, Changelist);
+		Refresh();
+	}
 }
 
 bool SSourceControlChangelistsWidget::CanSubmitChangelist()
