Added asset name & type, simlar to the save dialog to the changelists view.

diff --git a/Engine/Source/Editor/SourceControlWindows/Private/SSourceControlChangelists.cpp b/Engine/Source/Editor/SourceControlWindows/Private/SSourceControlChangelists.cpp
--- a/Engine/Source/Editor/SourceControlWindows/Private/SSourceControlChangelists.cpp
+++ b/Engine/Source/Editor/SourceControlWindows/Private/SSourceControlChangelists.cpp
@@ -22,12 +22,12 @@
 #include "ContentBrowserModule.h"
 #include "IContentBrowserSingleton.h"
 #include "Misc/MessageDialog.h"
+#include "AssetRegistry/AssetRegistryModule.h"
 
 
 #define LOCTEXT_NAMESPACE "SourceControlChangelist"
 
 //////////////////////////////
-
 static TSharedRef<SWidget> GetSCCFileWidget(FSourceControlStateRef InFileState)
 {
 	const FSlateBrush* IconBrush = FEditorStyle::GetBrush("ContentBrowser.ColumnViewAssetIcon");
@@ -102,24 +102,85 @@ struct FFileTreeItem : public IChangelistTreeItem
 		: FileState(InFileState)
 	{
 		Type = IChangelistTreeItem::File;
-	}
 
-	FText GetDisplayText() const
-	{
 		FString Filename = FileState->GetFilename();
+		FText AssetName = LOCTEXT("SourceControl_DefaultAssetName", "None");
+		FText AssetPath;
+		FText AssetType = LOCTEXT("SourceControl_DefaultAssetType", "Unknown");
 		FString AssetPackageName;
+		FColor AssetColor = FColor(		// Copied from ContentBrowserCLR.cpp
+			127 + FColor::Red.R / 2,	// Desaturate the colors a bit (GB colors were too.. much)
+			127 + FColor::Red.G / 2,
+			127 + FColor::Red.B / 2,
+			200); // Opacity
 
 		if (FPackageName::TryConvertFilenameToLongPackageName(Filename, AssetPackageName))
 		{
-			return FText::FromString(AssetPackageName);
+			AssetPath = FText::FromString(AssetPackageName);
+
+			TArray<FAssetData> Assets;
+			FAssetRegistryModule& AssetRegistryModule = FModuleManager::LoadModuleChecked<FAssetRegistryModule>(TEXT("AssetRegistry"));
+			AssetRegistryModule.Get().GetAssetsByPackageName(*AssetPackageName, Assets);
+
+			if (Assets.Num())
+			{
+				AssetName = FText::FromString(Assets[0].AssetName.ToString());
+				AssetColor = FColor::White;
+
+				if (Assets.Num() > 1)
+				{
+					AssetType = LOCTEXT("SourceCOntrol_ManyAssetType", "Multiple Assets");
+				}
+				else
+				{
+					AssetType = FText::FromString(Assets[0].AssetClass.ToString());
+
+					const FAssetToolsModule& AssetToolsModule = FModuleManager::LoadModuleChecked<FAssetToolsModule>(TEXT("AssetTools"));
+					const TSharedPtr<IAssetTypeActions> AssetTypeActions = AssetToolsModule.Get().GetAssetTypeActionsForClass(Assets[0].GetClass()).Pin();
+					if (AssetTypeActions.IsValid())
+					{
+						AssetColor = AssetTypeActions->GetTypeColor();
+					}
+				}
+			}
 		}
 		else
 		{
-			return FText::FromString(Filename);
+			AssetPath = FText::FromString(Filename);
 		}
+
+		DisplayName = AssetName;
+		DisplayPath = AssetPath;
+		DisplayType = AssetType;
+		DisplayColor = AssetColor;
+	}
+
+	FText GetDisplayPath() const
+	{
+		return DisplayPath;
+	}
+
+	FText GetDisplayName() const
+	{
+		return DisplayName;
+	}
+
+	FText GetDisplayType() const
+	{
+		return DisplayType;
+	}
+
+	FSlateColor GetDisplayColor() const
+	{
+		return FSlateColor(DisplayColor);
 	}
 
 	FSourceControlStateRef FileState;
+private:
+	FText DisplayPath;
+	FText DisplayName;
+	FText DisplayType;
+	FColor DisplayColor;
 };
 
 SSourceControlChangelistsWidget::SSourceControlChangelistsWidget()
@@ -664,6 +725,10 @@ TSharedRef<SChangelistTree> SSourceControlChangelistsWidget::CreateTreeviewWidge
 			.FillWidth(0.2f)
 			+ SHeaderRow::Column("Description")
 			.DefaultLabel(LOCTEXT("Description", "Description"))
+			.FillWidth(0.6f)
+			+ SHeaderRow::Column("Type")
+			.DefaultLabel(LOCTEXT("Type", "Type"))
+			.FillWidth(0.2f)
 		);
 }
 
@@ -780,7 +845,7 @@ private:
 	FChangelistTreeItem* TreeItem;
 };
 
-class SFileTableRow : public STableRow<FChangelistTreeItemPtr>
+class SFileTableRow : public SMultiColumnTableRow<FChangelistTreeItemPtr>
 {
 public:
 	SLATE_BEGIN_ARGS(SFileTableRow)
@@ -800,9 +865,18 @@ public:
 	{
 		TreeItem = static_cast<FFileTreeItem*>(InArgs._TreeItemToVisualize.Get());
 
-		ChildSlot
-		[
-			SNew(SHorizontalBox)
+		auto Args = FSuperRowType::FArguments()
+			.OnDragDetected(InArgs._OnDragDetected)
+			.ShowSelection(true);
+		FSuperRowType::Construct(Args, InOwner);
+	}
+
+	// SMultiColumnTableRow overrides
+	virtual TSharedRef<SWidget> GenerateWidgetForColumn(const FName& ColumnName) override
+	{
+		if (ColumnName == TEXT("Change")) // eq. to name
+		{
+			return SNew(SHorizontalBox)
 
 			// Icon
 			+ SHorizontalBox::Slot()
@@ -817,23 +891,44 @@ public:
 				.VAlign(VAlign_Center)
 				[
 					SNew(STextBlock)
-					.Text(this, &SFileTableRow::GetDisplayText)
-				]
-		];
+					.Text(this, &SFileTableRow::GetDisplayName)
+				];
+		}
+		else if (ColumnName == TEXT("Description")) // eq. to path
+		{
+			return SNew(STextBlock)
+				.Text(this, &SFileTableRow::GetDisplayPath);
+		}
+		else if (ColumnName == TEXT("Type"))
+		{
+			return SNew(STextBlock)
+				.Text(this, &SFileTableRow::GetDisplayType)
+				.ColorAndOpacity(this, &SFileTableRow::GetDisplayColor);
+		}
+		else
+		{
+			return SNullWidget::NullWidget;
+		}
+	}
 
-		auto Args = STableRow::FArguments()
-			.OnDragDetected(InArgs._OnDragDetected);
+	FText GetDisplayName() const
+	{
+		return TreeItem->GetDisplayName();
+	}
 
-		STableRow<FChangelistTreeItemPtr>::ConstructInternal(
-			Args
-			.ShowSelection(true),
-			InOwner
-		);
+	FText GetDisplayPath() const
+	{
+		return TreeItem->GetDisplayPath();
 	}
 
-	FText GetDisplayText() const
+	FText GetDisplayType() const
 	{
-		return TreeItem->GetDisplayText();
+		return TreeItem->GetDisplayType();
+	}
+
+	FSlateColor GetDisplayColor() const
+	{
+		return TreeItem->GetDisplayColor();
 	}
 
 protected:
