Added ActorLabel tag for actors using WP
Used the ActorLabel tag to obtain a better name in changelists dialog

diff --git a/Engine/Source/Editor/SourceControlWindows/Private/SSourceControlChangelists.cpp b/Engine/Source/Editor/SourceControlWindows/Private/SSourceControlChangelists.cpp
--- a/Engine/Source/Editor/SourceControlWindows/Private/SSourceControlChangelists.cpp
+++ b/Engine/Source/Editor/SourceControlWindows/Private/SSourceControlChangelists.cpp
@@ -96,6 +96,43 @@ struct FChangelistTreeItem : public IChangelistTreeItem
 	FSourceControlChangelistStateRef ChangelistState;
 };
 
+bool GetAssetData(const FString& InPackageName, const FString& InFileName, TArray<FAssetData>& OutAssets)
+{
+	OutAssets.Reset();
+
+	// Try the registry first
+	FAssetRegistryModule& AssetRegistryModule = FModuleManager::LoadModuleChecked<FAssetRegistryModule>(TEXT("AssetRegistry"));
+	AssetRegistryModule.Get().GetAssetsByPackageName(*InPackageName, OutAssets);
+
+	if (OutAssets.Num() > 0)
+	{
+		return true;
+	}
+
+	// Filter on improbable file extensions
+	EPackageExtension PackageExtension = FPackagePath::ParseExtension(InFileName);
+
+	if (PackageExtension == EPackageExtension::Unspecified ||
+		PackageExtension == EPackageExtension::Custom)
+	{
+		return false;
+	}
+
+	// If nothing was done, try to get the data explicitly
+	TArray<FAssetData*> LoadedAssets;
+	AssetRegistryModule.Get().LoadPackageRegistryData(InFileName, LoadedAssets);
+
+	for (FAssetData* AssetData : LoadedAssets)
+	{
+		OutAssets.Add(MoveTemp(*AssetData));
+		delete AssetData;
+	}
+
+	LoadedAssets.Reset();
+
+	return OutAssets.Num() > 0;
+}
+
 struct FFileTreeItem : public IChangelistTreeItem
 {
 	FFileTreeItem(FSourceControlStateRef InFileState)
@@ -103,55 +140,72 @@ struct FFileTreeItem : public IChangelistTreeItem
 	{
 		Type = IChangelistTreeItem::File;
 
+		// Initialize asset data first
 		FString Filename = FileState->GetFilename();
-		FText AssetName = LOCTEXT("SourceControl_DefaultAssetName", "None");
-		FText AssetPath;
-		FText AssetType = LOCTEXT("SourceControl_DefaultAssetType", "Unknown");
 		FString AssetPackageName;
+
+		if (FPackageName::TryConvertFilenameToLongPackageName(Filename, AssetPackageName))
+		{
+			::GetAssetData(AssetPackageName, Filename, Assets);
+		}
+
+		// Initialize display-related members
+		FString AssetName = LOCTEXT("SourceControl_DefaultAssetName", "None").ToString();
+		FString AssetType = LOCTEXT("SourceControl_DefaultAssetType", "Unknown").ToString();
+		FString AssetPath = Filename;
 		FColor AssetColor = FColor(		// Copied from ContentBrowserCLR.cpp
 			127 + FColor::Red.R / 2,	// Desaturate the colors a bit (GB colors were too.. much)
 			127 + FColor::Red.G / 2,
 			127 + FColor::Red.B / 2,
 			200); // Opacity
 
-		if (FPackageName::TryConvertFilenameToLongPackageName(Filename, AssetPackageName))
+		if(Assets.Num() > 0)
 		{
-			AssetPath = FText::FromString(AssetPackageName);
+			AssetPath = Assets[0].ObjectPath.ToString();
 
-			TArray<FAssetData> Assets;
-			FAssetRegistryModule& AssetRegistryModule = FModuleManager::LoadModuleChecked<FAssetRegistryModule>(TEXT("AssetRegistry"));
-			AssetRegistryModule.Get().GetAssetsByPackageName(*AssetPackageName, Assets);
+			// Strip asset name from object path
+			int32 LastDot = -1;
+			if (AssetPath.FindLastChar('.', LastDot))
+			{
+				AssetPath.LeftInline(LastDot);
+			}
 
-			if (Assets.Num())
+			// Find name, asset type & color only if there is exactly one asset
+			if (Assets.Num() == 1)
 			{
-				AssetName = FText::FromString(Assets[0].AssetName.ToString());
-				AssetColor = FColor::White;
+				static FName NAME_ActorLabel(TEXT("ActorLabel"));
+				if (Assets[0].FindTag(NAME_ActorLabel))
+				{
+					Assets[0].GetTagValue(NAME_ActorLabel, AssetName);
+				}
+				else
+				{
+					AssetName = Assets[0].AssetName.ToString();
+				}
+
+				AssetType = Assets[0].AssetClass.ToString();
 
-				if (Assets.Num() > 1)
+				const FAssetToolsModule& AssetToolsModule = FModuleManager::LoadModuleChecked<FAssetToolsModule>(TEXT("AssetTools"));
+				const TSharedPtr<IAssetTypeActions> AssetTypeActions = AssetToolsModule.Get().GetAssetTypeActionsForClass(Assets[0].GetClass()).Pin();
+				if (AssetTypeActions.IsValid())
 				{
-					AssetType = LOCTEXT("SourceCOntrol_ManyAssetType", "Multiple Assets");
+					AssetColor = AssetTypeActions->GetTypeColor();
 				}
 				else
 				{
-					AssetType = FText::FromString(Assets[0].AssetClass.ToString());
-
-					const FAssetToolsModule& AssetToolsModule = FModuleManager::LoadModuleChecked<FAssetToolsModule>(TEXT("AssetTools"));
-					const TSharedPtr<IAssetTypeActions> AssetTypeActions = AssetToolsModule.Get().GetAssetTypeActionsForClass(Assets[0].GetClass()).Pin();
-					if (AssetTypeActions.IsValid())
-					{
-						AssetColor = AssetTypeActions->GetTypeColor();
-					}
+					AssetColor = FColor::White;
 				}
 			}
-		}
-		else
-		{
-			AssetPath = FText::FromString(Filename);
+			else
+			{
+				AssetType = LOCTEXT("SourceCOntrol_ManyAssetType", "Multiple Assets").ToString();
+				AssetColor = FColor::White;
+			}
 		}
 
-		DisplayName = AssetName;
-		DisplayPath = AssetPath;
-		DisplayType = AssetType;
+		DisplayName = FText::FromString(AssetName);
+		DisplayPath = FText::FromString(AssetPath);
+		DisplayType = FText::FromString(AssetType);
 		DisplayColor = AssetColor;
 	}
 
@@ -175,8 +229,17 @@ struct FFileTreeItem : public IChangelistTreeItem
 		return FSlateColor(DisplayColor);
 	}
 
+	const TArray<FAssetData>& GetAssetData() const
+	{
+		return Assets;
+	}
+
+public:
 	FSourceControlStateRef FileState;
+
 private:
+	TArray<FAssetData> Assets;
+
 	FText DisplayPath;
 	FText DisplayName;
 	FText DisplayType;
@@ -257,7 +320,6 @@ void SSourceControlChangelistsWidget::Refresh()
 	if (ISourceControlModule::Get().IsEnabled())
 	{
 		ISourceControlProvider& SourceControlProvider = ISourceControlModule::Get().GetProvider();
-
 		TArray<FSourceControlChangelistRef> Changelists = SourceControlProvider.GetChangelists(EStateCacheUsage::Use);
 
 		TArray<FSourceControlChangelistStateRef> ChangelistsStates;
@@ -564,29 +626,14 @@ bool SSourceControlChangelistsWidget::CanSubmitChangelist()
 
 void SSourceControlChangelistsWidget::OnLocateFile()
 { 
-	FAssetToolsModule& AssetToolsModule = FModuleManager::LoadModuleChecked<FAssetToolsModule>(TEXT("AssetTools"));
-
-	TArray<UObject*> AssetsToSync;
-
-	TArray<FString> SelectedFiles = GetSelectedFiles();
-	AssetsToSync.Reserve(SelectedFiles.Num());
+	TArray<FAssetData> AssetsToSync;
+	TArray<FChangelistTreeItemPtr> SelectedItems = TreeView->GetSelectedItems();
 
-	for (const FString& SelectedFile : SelectedFiles)
+	for (const FChangelistTreeItemPtr& SelectedItem : SelectedItems)
 	{
-		FString AssetPackageName;
-
-		if (FPackageName::TryConvertFilenameToLongPackageName(SelectedFile, AssetPackageName))
+		if (SelectedItem->GetTreeItemType() == IChangelistTreeItem::File)
 		{
-			UPackage* AssetPackage = LoadPackage(nullptr, *AssetPackageName, LOAD_ForDiff | LOAD_DisableCompileOnLoad);
-
-			// grab the asset from the package - we assume asset name matches file name
-			FString AssetName = FPaths::GetBaseFilename(SelectedFile);
-			UObject* SelectedAsset = FindObject<UObject>(AssetPackage, *AssetName);
-
-			if (SelectedAsset)
-			{
-				AssetsToSync.Add(SelectedAsset);
-			}
+			AssetsToSync.Append(StaticCastSharedPtr<FFileTreeItem>(SelectedItem)->GetAssetData());
 		}
 	}
 
diff --git a/Engine/Source/Editor/SourceControlWindows/Private/SSourceControlHistory.cpp b/Engine/Source/Editor/SourceControlWindows/Private/SSourceControlHistory.cpp
--- a/Engine/Source/Editor/SourceControlWindows/Private/SSourceControlHistory.cpp
+++ b/Engine/Source/Editor/SourceControlWindows/Private/SSourceControlHistory.cpp
@@ -1642,7 +1642,7 @@ void FSourceControlWindows::DisplayRevisionHistory( const TArray<FString>& InPac
 	}
 }
 
-void FSourceControlWindows::DiffAgainstWorkspace(const FString& InFileName)
+bool FSourceControlWindows::DiffAgainstWorkspace(const FString& InFileName)
 {
 	FAssetToolsModule& AssetToolsModule = FModuleManager::LoadModuleChecked<FAssetToolsModule>(TEXT("AssetTools"));
 
@@ -1656,10 +1656,17 @@ void FSourceControlWindows::DiffAgainstWorkspace(const FString& InFileName)
 		// grab the asset from the package - we assume asset name matches file name
 		FString AssetName = FPaths::GetBaseFilename(InFileName);
 		SelectedAsset = FindObject<UObject>(AssetPackage, *AssetName);
+
+		if (!SelectedAsset && AssetPackage)
+		{
+			SelectedAsset = AssetPackage->FindAssetInPackage();
+		}
 	}
 
 	if (SelectedAsset)
 	{
 		AssetToolsModule.Get().DiffAgainstDepot(SelectedAsset, AssetPackageName, SelectedAsset->GetName());
 	}
+
+	return !!SelectedAsset;
 }
diff --git a/Engine/Source/Editor/SourceControlWindows/Public/SourceControlWindows.h b/Engine/Source/Editor/SourceControlWindows/Public/SourceControlWindows.h
--- a/Engine/Source/Editor/SourceControlWindows/Public/SourceControlWindows.h
+++ b/Engine/Source/Editor/SourceControlWindows/Public/SourceControlWindows.h
@@ -98,8 +98,10 @@ public:
 	 * Displays file diff against workspace version
 	 * 
 	 * @param InFileName Name of the file to diff
+	 * 
+	 * @return true if the diff could be performed
 	 */
-	static void DiffAgainstWorkspace(const FString& InFileName);
+	static bool DiffAgainstWorkspace(const FString& InFileName);
 
 protected:
 	/** Callback for ChoosePackagesToCheckIn(), continues to bring up UI once source control operations are complete */
diff --git a/Engine/Source/Runtime/Engine/Private/Actor.cpp b/Engine/Source/Runtime/Engine/Private/Actor.cpp
--- a/Engine/Source/Runtime/Engine/Private/Actor.cpp
+++ b/Engine/Source/Runtime/Engine/Private/Actor.cpp
@@ -440,6 +440,12 @@ void AActor::GetAssetRegistryTags(TArray<FAssetRegistryTag>& OutTags) const
 			const FString ActorMetaData = FBase64::Encode(SerializedData);
 			static FName NAME_ActorMetaData(TEXT("ActorMetaData"));
 			OutTags.Add(UObject::FAssetRegistryTag(NAME_ActorMetaData, ActorMetaData, UObject::FAssetRegistryTag::TT_Hidden));
+
+			if (ActorLabel.Len() > 0)
+			{
+				static FName NAME_ActorLabel(TEXT("ActorLabel"));
+				OutTags.Add(UObject::FAssetRegistryTag(NAME_ActorLabel, ActorLabel, UObject::FAssetRegistryTag::TT_Hidden));
+			}
 		}
 	}
 }
