Added mechanism to retain tree expansion states around tree view rebuilds (changelists dialog)

diff --git a/Engine/Source/Editor/SourceControlWindows/Private/SSourceControlChangelists.cpp b/Engine/Source/Editor/SourceControlWindows/Private/SSourceControlChangelists.cpp
--- a/Engine/Source/Editor/SourceControlWindows/Private/SSourceControlChangelists.cpp
+++ b/Engine/Source/Editor/SourceControlWindows/Private/SSourceControlChangelists.cpp
@@ -350,6 +350,9 @@ void SSourceControlChangelistsWidget::Refresh()
 {
 	if (ISourceControlModule::Get().IsEnabled())
 	{
+		TMap<FSourceControlChangelistStateRef, ExpandedState> ExpandedStates;
+		SaveExpandedState(ExpandedStates);
+
 		ISourceControlProvider& SourceControlProvider = ISourceControlModule::Get().GetProvider();
 		TArray<FSourceControlChangelistRef> Changelists = SourceControlProvider.GetChangelists(EStateCacheUsage::Use);
 
@@ -383,6 +386,8 @@ void SSourceControlChangelistsWidget::Refresh()
 			ChangelistsNodes.Add(ChangelistTreeItem);
 		}
 
+		RestoreExpandedState(ExpandedStates);
+
 		TreeView->RequestTreeRefresh();
 	}
 	else
@@ -1303,5 +1308,63 @@ void SSourceControlChangelistsWidget::OnGetChildren(FChangelistTreeItemPtr InPar
 	}
 }
 
+void SSourceControlChangelistsWidget::SaveExpandedState(TMap<FSourceControlChangelistStateRef, ExpandedState>& ExpandedStates) const
+{
+	for (FChangelistTreeItemPtr Root : ChangelistsNodes)
+	{
+		if (Root->GetTreeItemType() != IChangelistTreeItem::Changelist)
+		{
+			continue;
+		}
+
+		bool bChangelistExpanded = TreeView->IsItemExpanded(Root);
+
+		bool bShelveExpanded = false;
+		for (FChangelistTreeItemPtr Child : Root->GetChildren())
+		{
+			if (Child->GetTreeItemType() == IChangelistTreeItem::ShelvedChangelist)
+			{
+				bShelveExpanded = TreeView->IsItemExpanded(Child);
+				break;
+			}
+		}
+
+		ExpandedState State;
+		State.bChangelistExpanded = bChangelistExpanded;
+		State.bShelveExpanded = bShelveExpanded;
+
+		ExpandedStates.Add(StaticCastSharedPtr<FChangelistTreeItem>(Root)->ChangelistState, State);
+	}
+}
+
+void SSourceControlChangelistsWidget::RestoreExpandedState(const TMap<FSourceControlChangelistStateRef, ExpandedState>& ExpandedStates)
+{
+	for (FChangelistTreeItemPtr Root : ChangelistsNodes)
+	{
+		if (Root->GetTreeItemType() != IChangelistTreeItem::Changelist)
+		{
+			continue;
+		}
+
+		FSourceControlChangelistStateRef ChangelistState = StaticCastSharedPtr<FChangelistTreeItem>(Root)->ChangelistState;
+		const ExpandedState* State = ExpandedStates.Find(ChangelistState);
+
+		if (!State)
+		{
+			continue;
+		}
+
+		TreeView->SetItemExpansion(Root, State->bChangelistExpanded);
+
+		for (FChangelistTreeItemPtr Child : Root->GetChildren())
+		{
+			if (Child->GetTreeItemType() == IChangelistTreeItem::ShelvedChangelist)
+			{
+				TreeView->SetItemExpansion(Child, State->bShelveExpanded);
+				break;
+			}
+		}
+	}
+}
 
 #undef LOCTEXT_NAMESPACE
\ No newline at end of file
diff --git a/Engine/Source/Editor/SourceControlWindows/Private/SSourceControlChangelists.h b/Engine/Source/Editor/SourceControlWindows/Private/SSourceControlChangelists.h
--- a/Engine/Source/Editor/SourceControlWindows/Private/SSourceControlChangelists.h
+++ b/Engine/Source/Editor/SourceControlWindows/Private/SSourceControlChangelists.h
@@ -149,6 +149,14 @@ private:
 	virtual void Tick(const FGeometry& AllottedGeometry, const double InCurrentTime, const float InDeltaTime) override;
 
 private:
+	struct ExpandedState
+	{
+		bool bChangelistExpanded;
+		bool bShelveExpanded;
+	};
+
+	void SaveExpandedState(TMap<FSourceControlChangelistStateRef, ExpandedState>& ExpandedStates) const;
+	void RestoreExpandedState(const TMap<FSourceControlChangelistStateRef, ExpandedState>& ExpandedStates);
 
 	/** Changelists (root nodes) */
 	TArray<FChangelistTreeItemPtr> ChangelistsNodes;
