Added refresh button to the changelists dialog.

diff --git a/Engine/Source/Editor/SourceControlWindows/Private/SSourceControlChangelists.cpp b/Engine/Source/Editor/SourceControlWindows/Private/SSourceControlChangelists.cpp
--- a/Engine/Source/Editor/SourceControlWindows/Private/SSourceControlChangelists.cpp
+++ b/Engine/Source/Editor/SourceControlWindows/Private/SSourceControlChangelists.cpp
@@ -10,6 +10,9 @@
 #include "Widgets/Images/SImage.h"
 #include "Widgets/Layout/SScrollBorder.h"
 #include "Widgets/SOverlay.h"
+#include "Widgets/Input/SButton.h"
+#include "Widgets/Layout/SBorder.h"
+#include "Widgets/Layout/SBox.h"
 
 #include "ISourceControlProvider.h"
 #include "ISourceControlModule.h"
@@ -291,16 +294,54 @@ void SSourceControlChangelistsWidget::Construct(const FArguments& InArgs)
 
 	ChildSlot
 	[
-		SNew(SScrollBorder, TreeView.ToSharedRef())
-		.Visibility(TAttribute<EVisibility>::Create(TAttribute<EVisibility>::FGetter::CreateLambda([]()->EVisibility { return ISourceControlModule::Get().IsEnabled() ? EVisibility::Visible : EVisibility::Hidden; })))
+		SNew(SVerticalBox)
+		+ SVerticalBox::Slot()
+		.AutoHeight()
 		[
-			TreeView.ToSharedRef()
+			SNew(SBorder)
+			.BorderImage(FEditorStyle::GetBrush("ToolPanel.GroupBorder"))
+			.Padding(4)
+			[
+				SNew(SHorizontalBox)
+				+SHorizontalBox::Slot()
+				.HAlign(HAlign_Left)
+				.VAlign(VAlign_Center)
+				.AutoWidth()
+				[
+					MakeToolBar()
+				]
+			]
+		]
+		+ SVerticalBox::Slot()
+		[
+			SNew(SScrollBorder, TreeView.ToSharedRef())
+			.Visibility(TAttribute<EVisibility>::Create(TAttribute<EVisibility>::FGetter::CreateLambda([]()->EVisibility { return ISourceControlModule::Get().IsEnabled() ? EVisibility::Visible : EVisibility::Hidden; })))
+			[
+				TreeView.ToSharedRef()
+			]
 		]
 	];
 
 	bShouldRefresh = true;
 }
 
+TSharedRef<SWidget> SSourceControlChangelistsWidget::MakeToolBar()
+{
+	FSlimHorizontalToolBarBuilder ToolBarBuilder(nullptr, FMultiBoxCustomization::None);
+
+	ToolBarBuilder.AddToolBarButton(
+		FUIAction(
+			FExecuteAction::CreateLambda([this]() {
+				RequestRefresh();
+				})),
+		NAME_None,
+		LOCTEXT("SourceControl_RefreshButton", "Refresh"),
+		LOCTEXT("SourceControl_RefreshButton_Tooltip", "Refreshes changelists from source control provider."),
+		FSlateIcon(FEditorStyle::GetStyleSetName(), "SourceControl.Actions.Refresh"));
+
+	return ToolBarBuilder.MakeWidget();
+}
+
 void SSourceControlChangelistsWidget::Tick(const FGeometry& AllottedGeometry, const double InCurrentTime, const float InDeltaTime)
 {
 	if (bShouldRefresh)
diff --git a/Engine/Source/Editor/SourceControlWindows/Private/SSourceControlChangelists.h b/Engine/Source/Editor/SourceControlWindows/Private/SSourceControlChangelists.h
--- a/Engine/Source/Editor/SourceControlWindows/Private/SSourceControlChangelists.h
+++ b/Engine/Source/Editor/SourceControlWindows/Private/SSourceControlChangelists.h
@@ -158,6 +158,9 @@ private:
 	void SaveExpandedState(TMap<FSourceControlChangelistStateRef, ExpandedState>& ExpandedStates) const;
 	void RestoreExpandedState(const TMap<FSourceControlChangelistStateRef, ExpandedState>& ExpandedStates);
 
+	TSharedRef<SWidget> MakeToolBar();
+
+private:
 	/** Changelists (root nodes) */
 	TArray<FChangelistTreeItemPtr> ChangelistsNodes;
 
