Added data validation capabilities in the source control.
Exposed data validation in changelists

diff --git a/Engine/Plugins/Editor/DataValidation/Source/DataValidation/DataValidation.Build.cs b/Engine/Plugins/Editor/DataValidation/Source/DataValidation/DataValidation.Build.cs
--- a/Engine/Plugins/Editor/DataValidation/Source/DataValidation/DataValidation.Build.cs
+++ b/Engine/Plugins/Editor/DataValidation/Source/DataValidation/DataValidation.Build.cs
@@ -37,6 +37,7 @@ public class DataValidation : ModuleRules
 				"SandboxFile",
 				"Blutility",
 				"ToolMenus",
+				"SourceControl"
 			}
 		);
 	}
diff --git a/Engine/Plugins/Editor/DataValidation/Source/DataValidation/Private/DataValidationChangelist.cpp b/Engine/Plugins/Editor/DataValidation/Source/DataValidation/Private/DataValidationChangelist.cpp
new file mode 100644
--- /dev/null
+++ b/Engine/Plugins/Editor/DataValidation/Source/DataValidation/Private/DataValidationChangelist.cpp
@@ -0,0 +1,126 @@
+// Copyright Epic Games, Inc. All Rights Reserved.
+
+#include "DataValidationChangelist.h"
+#include "AssetData.h"
+#include "AssetRegistryModule.h"
+#include "ISourceControlProvider.h"
+#include "ISourceControlModule.h"
+#include "SourceControlHelpers.h"
+#include "Algo/AnyOf.h"
+
+#define LOCTEXT_NAMESPACE "DataValidationChangelist"
+
+void GatherDependencies(const FName& InPackageName, TSet<FName>& OutDependencies)
+{
+	if (OutDependencies.Contains(InPackageName))
+		return;
+
+	OutDependencies.Add(InPackageName);
+
+	TArray<FAssetData> Assets;
+	TArray<FName> Dependencies;
+	USourceControlHelpers::GetAssetDataFromPackage(InPackageName.ToString(), Assets, &Dependencies);
+
+	for (const FName& PackageDependency : Dependencies)
+	{
+		GatherDependencies(PackageDependency, OutDependencies);
+	}
+}
+
+FString GetPrettyPackageName(const FName& InPackageName)
+{
+	TArray<FAssetData> Assets;
+	USourceControlHelpers::GetAssetDataFromPackage(InPackageName.ToString(), Assets);
+
+	if (Assets.Num() > 0)
+	{
+		FString AssetPath = Assets[0].ObjectPath.ToString();
+
+		int32 LastDot = -1;
+		if (AssetPath.FindLastChar('.', LastDot))
+		{
+			AssetPath.LeftInline(LastDot);
+		}
+
+		FString AssetName;
+
+		static FName NAME_ActorLabel(TEXT("ActorLabel"));
+		if (Assets[0].FindTag(NAME_ActorLabel))
+		{
+			Assets[0].GetTagValue(NAME_ActorLabel, AssetName);
+		}
+		else
+		{
+			AssetName = Assets[0].AssetName.ToString();
+		}
+
+		return AssetPath + "." + AssetName;
+	}
+	else
+	{
+		return InPackageName.ToString();
+	}
+}
+
+EDataValidationResult UDataValidationChangelist::IsDataValid(TArray<FText>& ValidationErrors)
+{
+	ISourceControlProvider& SourceControlProvider = ISourceControlModule::Get().GetProvider();
+	
+	FSourceControlChangelistStatePtr ChangelistState = SourceControlProvider.GetState(Changelist->AsShared(), EStateCacheUsage::Use);
+
+	// Gather dependencies of every file in the changelist
+	TArray<FName> FilesInChangelist;
+	TSet<FName> ExternalDependenciesSet;
+
+	for (const FSourceControlStateRef& File : ChangelistState->GetFilesStates())
+	{
+		FString PackageName;
+		if (FPackageName::TryConvertFilenameToLongPackageName(File->GetFilename(), PackageName))
+		{
+			FilesInChangelist.Add(*PackageName);
+			GatherDependencies(*PackageName, ExternalDependenciesSet);
+		}
+	}
+
+	// For every dependency in the external dependencies that is not in the changelist
+	// Make sure that the source control state is "not currently modified"
+	TArray<FName> ExternalDependencies = ExternalDependenciesSet.Array();
+
+	ExternalDependencies.RemoveAll([&FilesInChangelist](FName& File) -> bool {
+		return Algo::AnyOf(FilesInChangelist, [&File](const auto& FileInChangelist) {
+			return File == FileInChangelist;
+			});
+		});
+
+	bool bHasChangelistErrors = false;
+
+	for(FName& ExternalDependency : ExternalDependencies)
+	{
+		FString ExternalPackageFilename = USourceControlHelpers::PackageFilename(ExternalDependency.ToString());
+
+		FSourceControlStatePtr ExternalDependencyFileState = SourceControlProvider.GetState(ExternalPackageFilename, EStateCacheUsage::Use);
+
+		// Check if file is in cache; if it's not in the cache, then it's not currently changed.
+		if (!ExternalDependencyFileState)
+		{
+			continue;
+		}
+
+		// If file is checked out, added, it's not in this changelist, which is a problem
+		if (ExternalDependencyFileState->IsCheckedOut() || ExternalDependencyFileState->IsAdded())
+		{
+			bHasChangelistErrors = true;
+			FText CurrentError = FText::Format(LOCTEXT("DataValidation.Changelist.Error", "{0} is missing from this changelist."), FText::FromString(GetPrettyPackageName(ExternalDependency)));
+			ValidationErrors.Add(CurrentError);
+		}
+	}
+
+	return bHasChangelistErrors ? EDataValidationResult::Invalid : EDataValidationResult::Valid;
+}
+
+void UDataValidationChangelist::Initialize(FSourceControlChangelistPtr InChangelist)
+{
+	Changelist = InChangelist;
+}
+
+#undef LOCTEXT_NAMESPACE
\ No newline at end of file
diff --git a/Engine/Plugins/Editor/DataValidation/Source/DataValidation/Private/EditorValidatorSubsystem.cpp b/Engine/Plugins/Editor/DataValidation/Source/DataValidation/Private/EditorValidatorSubsystem.cpp
--- a/Engine/Plugins/Editor/DataValidation/Source/DataValidation/Private/EditorValidatorSubsystem.cpp
+++ b/Engine/Plugins/Editor/DataValidation/Source/DataValidation/Private/EditorValidatorSubsystem.cpp
@@ -16,6 +16,8 @@
 #include "Logging/MessageLog.h"
 #include "Misc/ScopedSlowTask.h"
 #include "AssetData.h"
+#include "ISourceControlModule.h"
+#include "DataValidationChangelist.h"
 
 #define LOCTEXT_NAMESPACE "EditorValidationSubsystem"
 
@@ -69,6 +71,9 @@ void UEditorValidatorSubsystem::Initialize(FSubsystemCollectionBase& Collection)
 			}
 		}
 	}
+
+	// Register to SCC pre-submit callback
+	ISourceControlModule::Get().RegisterPreSubmitDataValidation(FSourceControlPreSubmitDataValidationDelegate::CreateUObject(this, &UEditorValidatorSubsystem::ValidateChangelistPreSubmit));
 }
 
 // Rename to BP validators
@@ -121,6 +126,9 @@ void UEditorValidatorSubsystem::Deinitialize()
 {
 	CleanupValidators();
 
+	// Unregister to SCC pre-submit callback
+	ISourceControlModule::Get().UnregisterPreSubmitDataValidation();
+
 	Super::Deinitialize();
 }
 
@@ -406,4 +414,23 @@ void UEditorValidatorSubsystem::ValidateAllSavedPackages()
 	SavedPackagesToValidate.Empty();
 }
 
+void UEditorValidatorSubsystem::ValidateChangelistPreSubmit(FSourceControlChangelistPtr InChangelist, EDataValidationResult& OutResult, TArray<FText>& OutValidationErrors, TArray<FText>& OutValidationWarnings) const
+{
+	FScopedSlowTask SlowTask(1.0f, LOCTEXT("ValidatingDataTask", "Validating changelist..."));
+	SlowTask.Visibility = ESlowTaskVisibility::Invisible;
+	SlowTask.MakeDialogDelayed(.1f);
+
+	check(InChangelist);
+
+	// Create temporary changelist object to do most of the heavy lifting
+	UDataValidationChangelist* Changelist = NewObject<UDataValidationChangelist>();
+	Changelist->Initialize(InChangelist);
+	Changelist->AddToRoot();
+
+	SlowTask.EnterProgressFrame(1.0f, LOCTEXT("ValidatingChangelist", "Validating changelist to submit"));
+	OutResult = IsObjectValid(Changelist, OutValidationErrors, OutValidationWarnings);
+
+	Changelist->RemoveFromRoot();
+}
+
 #undef LOCTEXT_NAMESPACE
diff --git a/Engine/Plugins/Editor/DataValidation/Source/DataValidation/Public/DataValidationChangelist.h b/Engine/Plugins/Editor/DataValidation/Source/DataValidation/Public/DataValidationChangelist.h
new file mode 100644
--- /dev/null
+++ b/Engine/Plugins/Editor/DataValidation/Source/DataValidation/Public/DataValidationChangelist.h
@@ -0,0 +1,43 @@
+// Copyright Epic Games, Inc. All Rights Reserved.
+
+#pragma once
+
+#include "EditorSubsystem.h"
+
+#include "UObject/ObjectMacros.h"
+#include "UObject/UObjectGlobals.h"
+#include "UObject/Interface.h"
+#include "UObject/ObjectMacros.h"
+#include "Engine/EngineTypes.h"
+
+#include "ISourceControlChangelist.h"
+
+#include "DataValidationChangelist.generated.h"
+
+/**
+ * Changelist abstraction to allow changelist-level data validation
+ */
+UCLASS(config = Editor)
+class DATAVALIDATION_API UDataValidationChangelist : public UObject
+{
+	GENERATED_BODY()
+
+public:
+	/** Default constructor with nothing */
+	UDataValidationChangelist() = default;
+
+	virtual EDataValidationResult IsDataValid(TArray<FText>& ValidationErrors) override;
+
+	/** Initializes internal state so the validation can be done */
+	void Initialize(FSourceControlChangelistPtr InChangelist);
+
+public:
+	/** Changelist to validate */
+	FSourceControlChangelistPtr Changelist;
+
+	/** Assets contained in the changelist */
+	//TArray<FAssetData> AssetsInChangelist;
+
+	/** Change description */
+	//FText Description;
+};
\ No newline at end of file
diff --git a/Engine/Plugins/Editor/DataValidation/Source/DataValidation/Public/EditorValidatorSubsystem.h b/Engine/Plugins/Editor/DataValidation/Source/DataValidation/Public/EditorValidatorSubsystem.h
--- a/Engine/Plugins/Editor/DataValidation/Source/DataValidation/Public/EditorValidatorSubsystem.h
+++ b/Engine/Plugins/Editor/DataValidation/Source/DataValidation/Public/EditorValidatorSubsystem.h
@@ -17,6 +17,8 @@
 #include "EditorValidatorSubsystem.generated.h"
 
 struct FAssetData;
+class ISourceControlChangelist;
+typedef TSharedPtr<class ISourceControlChangelist, ESPMode::ThreadSafe> FSourceControlChangelistPtr;
 
 DECLARE_LOG_CATEGORY_EXTERN(LogContentValidation, Log, All);
 
@@ -110,6 +112,11 @@ protected:
 
 	void RegisterBlueprintValidators();
 
+	/**
+	 * Validates changelist before SCC submit operation
+	 */
+	void ValidateChangelistPreSubmit(FSourceControlChangelistPtr Changelist, EDataValidationResult& OutResult, TArray<FText>& ValidationErrors, TArray<FText>& ValidationWarnings) const;
+
 protected:
 	/**
 	 * Directories to ignore for data validation. Useful for test assets
diff --git a/Engine/Source/Developer/SourceControl/Private/SourceControlHelpers.cpp b/Engine/Source/Developer/SourceControl/Private/SourceControlHelpers.cpp
--- a/Engine/Source/Developer/SourceControl/Private/SourceControlHelpers.cpp
+++ b/Engine/Source/Developer/SourceControl/Private/SourceControlHelpers.cpp
@@ -13,6 +13,7 @@
 #include "UObject/Package.h"
 #include "Misc/PackageName.h"
 #include "Logging/MessageLog.h"
+#include "AssetRegistry/AssetRegistryModule.h"
 
 #define LOCTEXT_NAMESPACE "SourceControlHelpers"
 
@@ -1422,6 +1423,66 @@ const FString& USourceControlHelpers::GetGlobalSettingsIni()
 	return SourceControlGlobalSettingsIni;
 }
 
+bool USourceControlHelpers::GetAssetData(const FString& InFileName, TArray<FAssetData>& OutAssets, TArray<FName>* OutDependencies)
+{
+	FString PackageName;
+	if (FPackageName::TryConvertFilenameToLongPackageName(InFileName, PackageName))
+	{
+		return GetAssetData(InFileName, PackageName, OutAssets, OutDependencies);
+	}
+	else
+	{
+		return false;
+	}
+}
+
+bool USourceControlHelpers::GetAssetDataFromPackage(const FString& PackageName, TArray<FAssetData>& OutAssets, TArray<FName>* OutDependencies)
+{
+	return GetAssetData(PackageFilename(PackageName), PackageName, OutAssets, OutDependencies);
+}
+
+bool USourceControlHelpers::GetAssetData(const FString & InFileName, const FString& InPackageName, TArray<FAssetData>& OutAssets, TArray<FName>* OutDependencies)
+{
+	const bool bGetDependencies = (OutDependencies != nullptr);
+	OutAssets.Reset();
+
+	// Try the registry first
+	FAssetRegistryModule& AssetRegistryModule = FModuleManager::LoadModuleChecked<FAssetRegistryModule>(TEXT("AssetRegistry"));
+	AssetRegistryModule.Get().GetAssetsByPackageName(*InPackageName, OutAssets);
+
+	if (OutAssets.Num() > 0)
+	{
+		// Assets are already in the cache, we can query dependencies directly
+		if (bGetDependencies)
+		{
+			AssetRegistryModule.Get().GetDependencies(*InPackageName, *OutDependencies);
+		}
+
+		return true;
+	}
+
+	// Filter on improbable file extensions
+	EPackageExtension PackageExtension = FPackagePath::ParseExtension(InFileName);
+
+	if (PackageExtension == EPackageExtension::Unspecified ||
+		PackageExtension == EPackageExtension::Custom)
+	{
+		return false;
+	}
+
+	// If nothing was done, try to get the data explicitly	
+	IAssetRegistry::FLoadPackageRegistryData LoadedData(bGetDependencies);
+
+	AssetRegistryModule.Get().LoadPackageRegistryData(InFileName, LoadedData);
+	OutAssets = MoveTemp(LoadedData.Data);
+
+	if (bGetDependencies)
+	{
+		*OutDependencies = MoveTemp(LoadedData.DataDependencies);
+	}	
+
+	return OutAssets.Num() > 0;
+}
 
 FScopedSourceControl::FScopedSourceControl()
 {
diff --git a/Engine/Source/Developer/SourceControl/Private/SourceControlModule.cpp b/Engine/Source/Developer/SourceControl/Private/SourceControlModule.cpp
--- a/Engine/Source/Developer/SourceControl/Private/SourceControlModule.cpp
+++ b/Engine/Source/Developer/SourceControl/Private/SourceControlModule.cpp
@@ -445,6 +445,21 @@ void FSourceControlModule::UnregisterProviderChanged(FDelegateHandle Handle)
 	OnSourceControlProviderChanged.Remove(Handle);
 }
 
+void FSourceControlModule::RegisterPreSubmitDataValidation(const FSourceControlPreSubmitDataValidationDelegate& PreSubmitDataValidationDelegate)
+{
+	OnSourceControlPreSubmitDataValidation = PreSubmitDataValidationDelegate;
+}
+
+void FSourceControlModule::UnregisterPreSubmitDataValidation()
+{
+	OnSourceControlPreSubmitDataValidation = FSourceControlPreSubmitDataValidationDelegate();
+}
+
+FSourceControlPreSubmitDataValidationDelegate FSourceControlModule::GetRegisteredPreSubmitDataValidation()
+{
+	return OnSourceControlPreSubmitDataValidation;
+}
+
 IMPLEMENT_MODULE( FSourceControlModule, SourceControl );
 
 #undef LOCTEXT_NAMESPACE
diff --git a/Engine/Source/Developer/SourceControl/Private/SourceControlModule.h b/Engine/Source/Developer/SourceControl/Private/SourceControlModule.h
--- a/Engine/Source/Developer/SourceControl/Private/SourceControlModule.h
+++ b/Engine/Source/Developer/SourceControl/Private/SourceControlModule.h
@@ -39,6 +39,9 @@ public:
 	virtual void SetUseGlobalSettings(bool bIsUseGlobalSettings) override;
 	virtual FDelegateHandle RegisterProviderChanged(const FSourceControlProviderChanged::FDelegate& SourceControlProviderChanged) override;
 	virtual void UnregisterProviderChanged(FDelegateHandle Handle) override;
+	virtual void RegisterPreSubmitDataValidation(const FSourceControlPreSubmitDataValidationDelegate& PreSubmitDataValidationDelegate) override;
+	virtual void UnregisterPreSubmitDataValidation() override;
+	virtual FSourceControlPreSubmitDataValidationDelegate GetRegisteredPreSubmitDataValidation() override;
 
 	/** Save the settings to the ini file */
 	void SaveSettings();
@@ -119,4 +122,7 @@ private:
 
 	/** For notifying when the source provider is changed */
 	FSourceControlProviderChanged OnSourceControlProviderChanged;
+
+	/** To call when doing pre-submit data validation */
+	FSourceControlPreSubmitDataValidationDelegate OnSourceControlPreSubmitDataValidation;
 };
diff --git a/Engine/Source/Developer/SourceControl/Public/ISourceControlModule.h b/Engine/Source/Developer/SourceControl/Public/ISourceControlModule.h
--- a/Engine/Source/Developer/SourceControl/Public/ISourceControlModule.h
+++ b/Engine/Source/Developer/SourceControl/Public/ISourceControlModule.h
@@ -5,8 +5,11 @@
 #include "CoreMinimal.h"
 #include "Modules/ModuleInterface.h"
 #include "Modules/ModuleManager.h"
+#include "UObject/UObjectGlobals.h"
 
 class ISourceControlProvider;
+class ISourceControlChangelist;
+typedef TSharedPtr<class ISourceControlChangelist, ESPMode::ThreadSafe> FSourceControlChangelistPtr;
 
 SOURCECONTROL_API DECLARE_LOG_CATEGORY_EXTERN(LogSourceControl, Log, All);
 
@@ -16,6 +19,9 @@ DECLARE_DELEGATE_OneParam( FSourceControlLoginClosed, bool );
 /** Delegate called when the active source control provider is changed */
 DECLARE_MULTICAST_DELEGATE_TwoParams( FSourceControlProviderChanged, ISourceControlProvider& /*OldProvider*/, ISourceControlProvider& /*NewProvider*/ );
 
+/** Delegate called on pre-submit for data validation */
+DECLARE_DELEGATE_FourParams(FSourceControlPreSubmitDataValidationDelegate, FSourceControlChangelistPtr /*Changelist*/, EDataValidationResult& /*Result*/, TArray<FText>& /*ValidationErrors*/, TArray<FText>& /*ValidationWarnings*/);
+
 /**
  * The modality of the login window.
  */
@@ -131,6 +137,21 @@ public:
 	 */
 	virtual void UnregisterProviderChanged(FDelegateHandle Handle) = 0;
 
+	/**
+	 * Register a delegate to be called to validate asset changes before submitting changes
+	 */
+	virtual void RegisterPreSubmitDataValidation(const FSourceControlPreSubmitDataValidationDelegate& PreSubmitDataValidationDelegate) = 0;
+
+	/**
+	 * Unregister a delegate called before submitting changes
+	 */
+	virtual void UnregisterPreSubmitDataValidation() = 0;
+
+	/**
+	 * Gets currently registered delegates for pre-submit data validation
+	 */
+	virtual FSourceControlPreSubmitDataValidationDelegate GetRegisteredPreSubmitDataValidation() = 0;
+
 	/**
 	 * Gets a reference to the source control module instance.
 	 *
diff --git a/Engine/Source/Developer/SourceControl/Public/SourceControlHelpers.h b/Engine/Source/Developer/SourceControl/Public/SourceControlHelpers.h
--- a/Engine/Source/Developer/SourceControl/Public/SourceControlHelpers.h
+++ b/Engine/Source/Developer/SourceControl/Public/SourceControlHelpers.h
@@ -493,6 +493,24 @@ public:
 	 */
 	static const FString& GetGlobalSettingsIni();
 
+	/**
+	* Helper function to retrieve the FAssetData associated with a given file
+	* @return Whether the asset data could be retrieved or not
+	*/
+	static bool GetAssetData(const FString& InFileName, const FString& InPackageName, TArray<FAssetData>& OutAssets, TArray<FName>* OutDependencies = nullptr);
+
+	/**
+	 * Helper function to retrieve the FAssetData associated with a given file
+	 * @return Whether the asset data could be retrieved or not
+	 */
+	static bool GetAssetData(const FString& InFileName, TArray<FAssetData>& OutAssets, TArray<FName>* OutDependencies = nullptr);
+
+	/**
+	* Helper function to retrieve the FAssetData associated with a given package name
+	* @return Whether the asset data could be retrieved or not
+	*/
+	static bool GetAssetDataFromPackage(const FString& InPackageName, TArray<FAssetData>& OutAssets, TArray<FName>* OutDependencies = nullptr);
+
 };  // USourceControlHelpers
 
 
diff --git a/Engine/Source/Developer/SourceControl/SourceControl.Build.cs b/Engine/Source/Developer/SourceControl/SourceControl.Build.cs
--- a/Engine/Source/Developer/SourceControl/SourceControl.Build.cs
+++ b/Engine/Source/Developer/SourceControl/SourceControl.Build.cs
@@ -21,6 +21,7 @@ public class SourceControl : ModuleRules
 
 		PrivateDependencyModuleNames.AddRange(
 			new string[] {
+				"AssetRegistry",
 				"Slate",
 				"EditorStyle"
 			}
diff --git a/Engine/Source/Editor/SourceControlWindows/Private/SSourceControlChangelists.cpp b/Engine/Source/Editor/SourceControlWindows/Private/SSourceControlChangelists.cpp
--- a/Engine/Source/Editor/SourceControlWindows/Private/SSourceControlChangelists.cpp
+++ b/Engine/Source/Editor/SourceControlWindows/Private/SSourceControlChangelists.cpp
@@ -21,11 +21,15 @@
 #include "Widgets/Images/SLayeredImage.h"
 #include "SSourceControlDescription.h"
 #include "SourceControlWindows.h"
+#include "SourceControlHelpers.h"
 #include "AssetToolsModule.h"
 #include "ContentBrowserModule.h"
 #include "IContentBrowserSingleton.h"
 #include "Misc/MessageDialog.h"
-#include "AssetRegistry/AssetRegistryModule.h"
+#include "Algo/AnyOf.h"
+
+#include "SSourceControlSubmit.h"
+#include "Framework/Application/SlateApplication.h"
 
 
 #define LOCTEXT_NAMESPACE "SourceControlChangelist"
@@ -113,43 +117,6 @@ struct FShelvedChangelistTreeItem : public IChangelistTreeItem
 	}
 };
 
-bool GetAssetData(const FString& InPackageName, const FString& InFileName, TArray<FAssetData>& OutAssets)
-{
-	OutAssets.Reset();
-
-	// Try the registry first
-	FAssetRegistryModule& AssetRegistryModule = FModuleManager::LoadModuleChecked<FAssetRegistryModule>(TEXT("AssetRegistry"));
-	AssetRegistryModule.Get().GetAssetsByPackageName(*InPackageName, OutAssets);
-
-	if (OutAssets.Num() > 0)
-	{
-		return true;
-	}
-
-	// Filter on improbable file extensions
-	EPackageExtension PackageExtension = FPackagePath::ParseExtension(InFileName);
-
-	if (PackageExtension == EPackageExtension::Unspecified ||
-		PackageExtension == EPackageExtension::Custom)
-	{
-		return false;
-	}
-
-	// If nothing was done, try to get the data explicitly
-	TArray<FAssetData*> LoadedAssets;
-	AssetRegistryModule.Get().LoadPackageRegistryData(InFileName, LoadedAssets);
-
-	for (FAssetData* AssetData : LoadedAssets)
-	{
-		OutAssets.Add(MoveTemp(*AssetData));
-		delete AssetData;
-	}
-
-	LoadedAssets.Reset();
-
-	return OutAssets.Num() > 0;
-}
-
 struct FFileTreeItem : public IChangelistTreeItem
 {
 	FFileTreeItem(FSourceControlStateRef InFileState)
@@ -159,12 +126,7 @@ struct FFileTreeItem : public IChangelistTreeItem
 
 		// Initialize asset data first
 		FString Filename = FileState->GetFilename();
-		FString AssetPackageName;
-
-		if (FPackageName::TryConvertFilenameToLongPackageName(Filename, AssetPackageName))
-		{
-			::GetAssetData(AssetPackageName, Filename, Assets);
-		}
+		USourceControlHelpers::GetAssetData(Filename, Assets);
 
 		// Initialize display-related members
 		FString AssetName = LOCTEXT("SourceControl_DefaultAssetName", "None").ToString();
@@ -722,6 +684,62 @@ void SSourceControlChangelistsWidget::OnDeleteShelvedFiles()
 	SourceControlProvider.Execute(DeleteShelvedOperation, GetChangelistFromSelection(), GetSelectedShelvedFiles());
 }
 
+static void GetChangelistValidationResult(FSourceControlChangelistPtr InChangelist, FString& OutValidationText, FName& OutValidationIcon)
+{
+	FSourceControlPreSubmitDataValidationDelegate ValidationDelegate = ISourceControlModule::Get().GetRegisteredPreSubmitDataValidation();
+
+	EDataValidationResult ValidationResult = EDataValidationResult::NotValidated;
+	TArray<FText> ValidationErrors;
+	TArray<FText> ValidationWarnings;
+
+	if (ValidationDelegate.ExecuteIfBound(InChangelist, ValidationResult, ValidationErrors, ValidationWarnings))
+	{
+		if (ValidationResult == EDataValidationResult::Invalid || ValidationErrors.Num() > 0)
+		{
+			OutValidationIcon = "MessageLog.Error";
+			OutValidationText = LOCTEXT("SourceControl.Submit.ChangelistValidationError", "Changelist validation failed!").ToString();
+		}
+		else if (ValidationResult == EDataValidationResult::NotValidated || ValidationWarnings.Num() > 0)
+		{
+			OutValidationIcon = "MessageLog.Warning";
+			OutValidationText = LOCTEXT("SourceControl.Submit.ChangelistValidationWarning", "Changelist validation has warnings!").ToString();
+		}
+		else
+		{
+			OutValidationIcon = "MessageLog.Note";
+			OutValidationText = LOCTEXT("SourceControl.Submit.ChangelistValidationSuccess", "Changelist validation succesful!").ToString();
+		}
+
+		int32 NumLinesDisplayed = 0;
+
+		auto AppendInfo = [&OutValidationText, &NumLinesDisplayed](const TArray<FText>& Info, const FString& InfoType) {
+			const int32 MaxNumLinesDisplayed = 5;
+
+			if (Info.Num() > 0)
+			{
+				OutValidationText += LINE_TERMINATOR;
+				OutValidationText += FString::Printf(TEXT("Encountered %d %s:"), Info.Num(), *InfoType);
+
+				for (const FText& Line : Info)
+				{
+					if (NumLinesDisplayed >= MaxNumLinesDisplayed)
+					{
+						break;
+					}
+
+					OutValidationText += LINE_TERMINATOR;
+					OutValidationText += Line.ToString();
+
+					++NumLinesDisplayed;
+				}
+			}
+		};
+
+		AppendInfo(ValidationErrors, TEXT("errors"));
+		AppendInfo(ValidationWarnings, TEXT("warnings"));
+	}
+}
+
 void SSourceControlChangelistsWidget::OnSubmitChangelist()
 {
 	FSourceControlChangelistStatePtr ChangelistState = GetCurrentChangelistState();
@@ -731,40 +749,53 @@ void SSourceControlChangelistsWidget::OnSubmitChangelist()
 		return;
 	}
 
-	const FText DialogText = LOCTEXT("SourceControl_ConfirmSubmit", "Are you sure you want to submit this changelist?");
-	const FText DialogTitle = LOCTEXT("SourceControl_ConfirmSubmit_Title", "Confirm changelist submit");
+	FString ChangelistValidationText;
+	FName ChangelistValidationIconName;
+	GetChangelistValidationResult(ChangelistState->GetChangelist(), ChangelistValidationText, ChangelistValidationIconName);
 
-	EAppReturnType::Type UserConfirmation = FMessageDialog::Open(EAppMsgType::OkCancel, EAppReturnType::Ok, DialogText, &DialogTitle);
+	// Build list of states for the dialog
+	FText ChangelistDescription = ChangelistState->GetDescriptionText();
+	const bool bAskForChangelistDescription = (ChangelistDescription.IsEmptyOrWhitespace());
 
-	if (UserConfirmation == EAppReturnType::Ok)
-	{
-		FText ChangelistDescription = ChangelistState->GetDescriptionText();
-		const bool bAskForChangelistDescription = (ChangelistDescription.IsEmptyOrWhitespace());
+	TSharedRef<SWindow> NewWindow = SNew(SWindow)
+		.Title(NSLOCTEXT("SourceControl.ConfirmSubmit", "Title", "Confirm changelist submit"))
+		.SizingRule(ESizingRule::UserSized)
+		.ClientSize(FVector2D(600, 400))
+		.SupportsMaximize(true)
+		.SupportsMinimize(false);
 
-		if (bAskForChangelistDescription)
-		{
-			bool bOk = GetChangelistDescription(
-				nullptr,
-				LOCTEXT("SourceControl.Changelist.Submit.Description.Title", "Enter a submit description..."),
-				LOCTEXT("SourceControl.Changelist.Submit.Description.Label", "Enter a description for the submit:"),
-				ChangelistDescription);
+	TSharedRef<SSourceControlSubmitWidget> SourceControlWidget =
+		SNew(SSourceControlSubmitWidget)
+		.ParentWindow(NewWindow)
+		.Items(ChangelistState->GetFilesStates())
+		.Description(ChangelistDescription)
+		.ChangeValidationDescription(ChangelistValidationText)
+		.ChangeValidationIcon(ChangelistValidationIconName)
+		.AllowDescriptionChange(bAskForChangelistDescription)
+		.AllowUncheckFiles(false)
+		.AllowKeepCheckedOut(false);
 
-			if (!bOk)
-			{
-				return; // Abort submit
-			}
-		}
+	NewWindow->SetContent(
+		SourceControlWidget
+	);
 
+	FSlateApplication::Get().AddModalWindow(NewWindow, NULL);
+
+	if (SourceControlWidget->GetResult() == ESubmitResults::SUBMIT_ACCEPTED)
+	{
 		ISourceControlProvider& SourceControlProvider = ISourceControlModule::Get().GetProvider();
 		auto SubmitChangelistOperation = ISourceControlOperation::Create<FCheckIn>();
 
+		// Replace description only if there was none
 		if (bAskForChangelistDescription)
 		{
-			SubmitChangelistOperation->SetDescription(ChangelistDescription);
+			FChangeListDescription Description;
+			SourceControlWidget->FillChangeListDescription(Description);
+
+			SubmitChangelistOperation->SetDescription(Description.Description);
 		}
 
 		SourceControlProvider.Execute(SubmitChangelistOperation, ChangelistState->GetChangelist());
-		Refresh();
 	}
 }
 
diff --git a/Engine/Source/Editor/SourceControlWindows/Private/SSourceControlSubmit.cpp b/Engine/Source/Editor/SourceControlWindows/Private/SSourceControlSubmit.cpp
--- a/Engine/Source/Editor/SourceControlWindows/Private/SSourceControlSubmit.cpp
+++ b/Engine/Source/Editor/SourceControlWindows/Private/SSourceControlSubmit.cpp
@@ -103,6 +103,11 @@ void SSourceControlSubmitWidget::Construct(const FArguments& InArgs)
 	ParentFrame = InArgs._ParentWindow.Get();
 	SortByColumn = SSourceControlSubmitWidgetDefs::ColumnID_AssetLabel;
 	SortMode = EColumnSortMode::Ascending;
+	SavedChangeListDescription = InArgs._Description.Get();
+	const bool bDescriptionIsReadOnly = !InArgs._AllowDescriptionChange.Get();
+	const bool bAllowUncheckFiles = InArgs._AllowUncheckFiles.Get();
+	const bool bAllowKeepCheckedOut = InArgs._AllowKeepCheckedOut.Get();
+	const bool bShowChangelistValidation = !InArgs._ChangeValidationDescription.Get().IsEmpty();
 
 	for (const auto& Item : InArgs._Items.Get())
 	{
@@ -111,15 +116,18 @@ void SSourceControlSubmitWidget::Construct(const FArguments& InArgs)
 
 	TSharedRef<SHeaderRow> HeaderRowWidget = SNew(SHeaderRow);
 
-	HeaderRowWidget->AddColumn(
-		SHeaderRow::Column(SSourceControlSubmitWidgetDefs::ColumnID_CheckBoxLabel)
-		[
-			SNew(SCheckBox)
-			.IsChecked(this, &SSourceControlSubmitWidget::GetToggleSelectedState)
-			.OnCheckStateChanged(this, &SSourceControlSubmitWidget::OnToggleSelectedCheckBox)
-		]
-		.FixedWidth(SSourceControlSubmitWidgetDefs::CheckBoxColumnWidth)
-	);
+	if (bAllowUncheckFiles)
+	{
+		HeaderRowWidget->AddColumn(
+			SHeaderRow::Column(SSourceControlSubmitWidgetDefs::ColumnID_CheckBoxLabel)
+			[
+				SNew(SCheckBox)
+				.IsChecked(this, &SSourceControlSubmitWidget::GetToggleSelectedState)
+				.OnCheckStateChanged(this, &SSourceControlSubmitWidget::OnToggleSelectedCheckBox)
+			]
+			.FixedWidth(SSourceControlSubmitWidgetDefs::CheckBoxColumnWidth)
+		);
+	}
 
 	HeaderRowWidget->AddColumn(
 		SHeaderRow::Column(SSourceControlSubmitWidgetDefs::ColumnID_IconLabel)
@@ -147,106 +155,151 @@ void SSourceControlSubmitWidget::Construct(const FArguments& InArgs)
 		.FillWidth(7.0f)
 	);
 
+	TSharedPtr<SVerticalBox> Contents;
+
 	ChildSlot
 	[
 		SNew(SBorder)
 		.BorderImage(FEditorStyle::GetBrush("ToolPanel.GroupBorder"))
 		[
-			SNew(SVerticalBox)
-			+SVerticalBox::Slot()
-			.AutoHeight()
+			SAssignNew(Contents, SVerticalBox)
+		]
+	];
+
+	// Build contents of dialog
+	Contents->AddSlot()
+	.AutoHeight()
+	.Padding(5)
+	[
+		SNew(STextBlock)
+		.Text(NSLOCTEXT("SourceControl.SubmitPanel", "ChangeListDesc", "Changelist Description"))
+	];
+
+	Contents->AddSlot()
+	.FillHeight(.5f)
+	.Padding(FMargin(5, 0, 5, 5))
+	[
+		SNew(SBox)
+		.WidthOverride(520)
+		[
+			SAssignNew(ChangeListDescriptionTextCtrl, SMultiLineEditableTextBox)
+			.SelectAllTextWhenFocused(!bDescriptionIsReadOnly)
+			.Text(SavedChangeListDescription)
+			.AutoWrapText(true)
+			.IsReadOnly(bDescriptionIsReadOnly)
+		]
+	];
+
+	Contents->AddSlot()
+	.Padding(FMargin(5, 0))
+	[
+		SNew(SBorder)
+		[
+			SAssignNew(ListView, SListView<TSharedPtr<FSubmitItem>>)
+			.ItemHeight(20)
+			.ListItemsSource(&ListViewItems)
+			.OnGenerateRow(this, &SSourceControlSubmitWidget::OnGenerateRowForList)
+			.OnContextMenuOpening(this, &SSourceControlSubmitWidget::OnCreateContextMenu)
+			.OnMouseButtonDoubleClick(this, &SSourceControlSubmitWidget::OnDiffAgainstDepotSelected)
+			.HeaderRow(HeaderRowWidget)
+			.SelectionMode(ESelectionMode::Single)
+		]
+	];
+
+	if (!bDescriptionIsReadOnly)
+	{
+		Contents->AddSlot()
+		.AutoHeight()
+		.Padding(FMargin(5, 5, 5, 0))
+		[
+			SNew( SBorder)
+			.Visibility(this, &SSourceControlSubmitWidget::IsWarningPanelVisible)
 			.Padding(5)
 			[
-				SNew( STextBlock )
-				.Text( NSLOCTEXT("SourceControl.SubmitPanel", "ChangeListDesc", "Changelist Description") )
-			]
-			+SVerticalBox::Slot()
-			.FillHeight(.5f)
-			.Padding(FMargin(5, 0, 5, 5))
-			[
-				SNew(SBox)
-				.WidthOverride(520)
-				[
-					SAssignNew( ChangeListDescriptionTextCtrl, SMultiLineEditableTextBox )
-					.SelectAllTextWhenFocused( true )
-					.Text(SavedChangeListDescription)
-					.AutoWrapText( true )
-				]
+				SNew( SErrorText )
+				.ErrorText( NSLOCTEXT("SourceControl.SubmitPanel", "ChangeListDescWarning", "Changelist description is required to submit") )
 			]
-			+SVerticalBox::Slot()
-			.Padding(FMargin(5, 0))
-			[
-				SNew(SBorder)
-				[
-					SAssignNew(ListView, SListView<TSharedPtr<FSubmitItem>>)
-					.ItemHeight(20)
-					.ListItemsSource(&ListViewItems)
-					.OnGenerateRow(this, &SSourceControlSubmitWidget::OnGenerateRowForList)
-					.OnContextMenuOpening(this, &SSourceControlSubmitWidget::OnCreateContextMenu)
-					.OnMouseButtonDoubleClick(this, &SSourceControlSubmitWidget::OnDiffAgainstDepotSelected)
-					.HeaderRow(HeaderRowWidget)
-					.SelectionMode(ESelectionMode::Single)
-				]
-			]
-			+SVerticalBox::Slot()
-			.AutoHeight()
-			.Padding(FMargin(5, 5, 5, 0))
+		];
+	}
+
+	if (bShowChangelistValidation)
+	{
+		FString ChangelistResultText = InArgs._ChangeValidationDescription.Get();
+		FName ChangelistIconName = InArgs._ChangeValidationIcon.Get();
+
+		Contents->AddSlot()
+		.AutoHeight()
+		.Padding(FMargin(5))
+		[
+			SNew(SHorizontalBox)
+			+SHorizontalBox::Slot()
+			.AutoWidth()
+			.VAlign(VAlign_Center)
 			[
-				SNew( SBorder)
-				.Visibility(this, &SSourceControlSubmitWidget::IsWarningPanelVisible)
-				.Padding(5)
-				[
-					SNew( SErrorText )
-					.ErrorText( NSLOCTEXT("SourceControl.SubmitPanel", "ChangeListDescWarning", "Changelist description is required to submit") )
-				]
+				SNew(SImage)
+				.Image(FEditorStyle::GetBrush(ChangelistIconName))
 			]
-			+SVerticalBox::Slot()
-			.AutoHeight()
-			.Padding(5)
+			+SHorizontalBox::Slot()
 			[
-				SNew(SWrapBox)
-				.UseAllottedSize(true)
-				+SWrapBox::Slot()
-				.Padding(0.0f, 0.0f, 16.0f, 0.0f)
-				[
-					SNew(SCheckBox)
-					.OnCheckStateChanged( this, &SSourceControlSubmitWidget::OnCheckStateChanged_KeepCheckedOut)
-					.IsChecked( this, &SSourceControlSubmitWidget::GetKeepCheckedOut )
-					.IsEnabled( this, &SSourceControlSubmitWidget::CanCheckOut )
-					[
-						SNew(STextBlock)
-						.Text(NSLOCTEXT("SourceControl.SubmitPanel", "KeepCheckedOut", "Keep Files Checked Out") )
-					]
-				]
+				SNew(SMultiLineEditableTextBox)
+				.Text(FText::FromString(ChangelistResultText))
+				.AutoWrapText(true)
+				.IsReadOnly(true)
 			]
-			+SVerticalBox::Slot()
-			.AutoHeight()
-			.HAlign(HAlign_Right)
-			.VAlign(VAlign_Bottom)
-			.Padding(0.0f,0.0f,0.0f,5.0f)
+		];
+	}
+
+	if (bAllowKeepCheckedOut)
+	{
+		Contents->AddSlot()
+		.AutoHeight()
+		.Padding(5)
+		[
+			SNew(SWrapBox)
+			.UseAllottedSize(true)
+			+SWrapBox::Slot()
+			.Padding(0.0f, 0.0f, 16.0f, 0.0f)
 			[
-				SNew(SUniformGridPanel)
-				.SlotPadding(FEditorStyle::GetMargin("StandardDialog.SlotPadding"))
-				.MinDesiredSlotWidth(FEditorStyle::GetFloat("StandardDialog.MinDesiredSlotWidth"))
-				.MinDesiredSlotHeight(FEditorStyle::GetFloat("StandardDialog.MinDesiredSlotHeight"))
-				+SUniformGridPanel::Slot(0,0)
-				[
-					SNew(SButton)
-					.HAlign(HAlign_Center)
-					.ContentPadding(FEditorStyle::GetMargin("StandardDialog.ContentPadding"))
-					.IsEnabled(this, &SSourceControlSubmitWidget::IsOKEnabled)
-					.Text( NSLOCTEXT("SourceControl.SubmitPanel", "OKButton", "Submit") )
-					.OnClicked(this, &SSourceControlSubmitWidget::OKClicked)
-				]
-				+SUniformGridPanel::Slot(1,0)
+				SNew(SCheckBox)
+				.OnCheckStateChanged( this, &SSourceControlSubmitWidget::OnCheckStateChanged_KeepCheckedOut)
+				.IsChecked( this, &SSourceControlSubmitWidget::GetKeepCheckedOut )
+				.IsEnabled( this, &SSourceControlSubmitWidget::CanCheckOut )
 				[
-					SNew(SButton)
-					.HAlign(HAlign_Center)
-					.ContentPadding(FEditorStyle::GetMargin("StandardDialog.ContentPadding"))
-					.Text( NSLOCTEXT("SourceControl.SubmitPanel", "CancelButton", "Cancel") )
-					.OnClicked(this, &SSourceControlSubmitWidget::CancelClicked)
+					SNew(STextBlock)
+					.Text(NSLOCTEXT("SourceControl.SubmitPanel", "KeepCheckedOut", "Keep Files Checked Out") )
 				]
 			]
+		];
+	}
+
+	const float AdditionalTopPadding = (bAllowKeepCheckedOut ? 0.0f : 5.0f);
+
+	Contents->AddSlot()
+	.AutoHeight()
+	.HAlign(HAlign_Right)
+	.VAlign(VAlign_Bottom)
+	.Padding(0.0f,AdditionalTopPadding,0.0f,5.0f)
+	[
+		SNew(SUniformGridPanel)
+		.SlotPadding(FEditorStyle::GetMargin("StandardDialog.SlotPadding"))
+		.MinDesiredSlotWidth(FEditorStyle::GetFloat("StandardDialog.MinDesiredSlotWidth"))
+		.MinDesiredSlotHeight(FEditorStyle::GetFloat("StandardDialog.MinDesiredSlotHeight"))
+		+SUniformGridPanel::Slot(0,0)
+		[
+			SNew(SButton)
+			.HAlign(HAlign_Center)
+			.ContentPadding(FEditorStyle::GetMargin("StandardDialog.ContentPadding"))
+			.IsEnabled(this, &SSourceControlSubmitWidget::IsOKEnabled)
+			.Text( NSLOCTEXT("SourceControl.SubmitPanel", "OKButton", "Submit") )
+			.OnClicked(this, &SSourceControlSubmitWidget::OKClicked)
+		]
+		+SUniformGridPanel::Slot(1,0)
+		[
+			SNew(SButton)
+			.HAlign(HAlign_Center)
+			.ContentPadding(FEditorStyle::GetMargin("StandardDialog.ContentPadding"))
+			.Text( NSLOCTEXT("SourceControl.SubmitPanel", "CancelButton", "Cancel") )
+			.OnClicked(this, &SSourceControlSubmitWidget::CancelClicked)
 		]
 	];
 
diff --git a/Engine/Source/Editor/SourceControlWindows/Private/SSourceControlSubmit.h b/Engine/Source/Editor/SourceControlWindows/Private/SSourceControlSubmit.h
--- a/Engine/Source/Editor/SourceControlWindows/Private/SSourceControlSubmit.h
+++ b/Engine/Source/Editor/SourceControlWindows/Private/SSourceControlSubmit.h
@@ -101,10 +101,22 @@ public:
 	SLATE_BEGIN_ARGS(SSourceControlSubmitWidget)
 		: _ParentWindow()
 		, _Items()
+		, _Description()
+		, _ChangeValidationIcon()
+		, _ChangeValidationDescription()
+		, _AllowDescriptionChange(true)
+		, _AllowUncheckFiles(true)
+		, _AllowKeepCheckedOut(true)
 	{}
 
 		SLATE_ATTRIBUTE(TSharedPtr<SWindow>, ParentWindow)
 		SLATE_ATTRIBUTE(TArray<FSourceControlStateRef>, Items)
+		SLATE_ATTRIBUTE(FText, Description)
+		SLATE_ATTRIBUTE(FName, ChangeValidationIcon)
+		SLATE_ATTRIBUTE(FString, ChangeValidationDescription)
+		SLATE_ATTRIBUTE(bool, AllowDescriptionChange)
+		SLATE_ATTRIBUTE(bool, AllowUncheckFiles)
+		SLATE_ATTRIBUTE(bool, AllowKeepCheckedOut)
 
 	SLATE_END_ARGS()
 
diff --git a/Engine/Source/Editor/UnrealEd/Private/Commandlets/DiffFilesCommandlet.cpp b/Engine/Source/Editor/UnrealEd/Private/Commandlets/DiffFilesCommandlet.cpp
--- a/Engine/Source/Editor/UnrealEd/Private/Commandlets/DiffFilesCommandlet.cpp
+++ b/Engine/Source/Editor/UnrealEd/Private/Commandlets/DiffFilesCommandlet.cpp
@@ -150,7 +150,7 @@ void UDiffFilesCommandlet::LoadAndDiff()
 	{
 		FAssetRegistryModule& AssetRegistryModule = FModuleManager::LoadModuleChecked<FAssetRegistryModule>(TEXT("AssetRegistry"));
 		IAssetRegistry& AssetRegistry = AssetRegistryModule.Get();
-		TArray<FAssetData*> AssetData;
+		IAssetRegistry::FLoadPackageRegistryData AssetData;
 		AssetRegistry.LoadPackageRegistryData(*Package->LinkerLoad->GetLoader_Unsafe(), AssetData);
 	}
 }
diff --git a/Engine/Source/Runtime/AssetRegistry/Private/AssetRegistry.cpp b/Engine/Source/Runtime/AssetRegistry/Private/AssetRegistry.cpp
--- a/Engine/Source/Runtime/AssetRegistry/Private/AssetRegistry.cpp
+++ b/Engine/Source/Runtime/AssetRegistry/Private/AssetRegistry.cpp
@@ -2133,30 +2133,62 @@ uint32 UAssetRegistryImpl::GetAllocatedSize(bool bLogDetailed) const
 	return StateSize + StaticSize + SearchSize;
 }
 
-void UAssetRegistryImpl::LoadPackageRegistryData(FArchive& Ar, TArray<FAssetData*> &AssetDataList) const
+void UAssetRegistryImpl::LoadPackageRegistryData(FArchive& Ar, FLoadPackageRegistryData& InOutData) const
 {
 	FPackageReader Reader;
 	if (Reader.OpenPackageFile(&Ar))
 	{
-		ReadAssetFile(Reader, AssetDataList);
+		ReadAssetFile(Reader, InOutData);
 	}
 }
 
-void UAssetRegistryImpl::LoadPackageRegistryData(const FString& PackageFilename, TArray<FAssetData*>& AssetDataList) const
+void UAssetRegistryImpl::LoadPackageRegistryData(const FString& PackageFilename, FLoadPackageRegistryData& InOutData) const
 {
 	FPackageReader Reader;
 	if (Reader.OpenPackageFile(PackageFilename))
 	{
-		ReadAssetFile(Reader, AssetDataList);
+		ReadAssetFile(Reader, InOutData);
 	}
 }
 
-bool UAssetRegistryImpl::ReadAssetFile(FPackageReader& PackageReader, TArray<FAssetData*>& AssetDataList) const
+bool UAssetRegistryImpl::ReadAssetFile(FPackageReader& PackageReader, FLoadPackageRegistryData& InOutData) const
 {
+	TArray<FAssetData*> AssetDataList;
 	TArray<FString> CookedPackageNamesWithoutAssetDataGathered;
-	// Dependency Data not currently read for this function
-	FPackageDependencyData* DependencyData = nullptr;
-	return FAssetDataGatherer::ReadAssetFile(PackageReader, AssetDataList, DependencyData, CookedPackageNamesWithoutAssetDataGathered);
+
+	const bool bGetDependencies = (InOutData.bGetDependencies);
+	FPackageDependencyData DependencyData;
+
+	bool bReadOk = FAssetDataGatherer::ReadAssetFile(PackageReader, AssetDataList, (InOutData.bGetDependencies ? &DependencyData : nullptr), CookedPackageNamesWithoutAssetDataGathered);
+
+	if (bReadOk)
+	{
+		// Copy & free asset data to the InOutData
+		InOutData.Data.Reset(AssetDataList.Num());
+		for (FAssetData* AssetData : AssetDataList)
+		{
+			InOutData.Data.Emplace(*AssetData);
+		}
+
+		AssetDataList.Reset();
+
+		if (InOutData.bGetDependencies)
+		{
+			InOutData.DataDependencies.Reset(DependencyData.ImportMap.Num());
+			for (int32 ImportIdx = 0; ImportIdx < DependencyData.ImportMap.Num(); ++ImportIdx)
+			{
+				InOutData.DataDependencies.Emplace(DependencyData.GetImportPackageName(ImportIdx));
+			}
+		}
+	}
+
+	// Cleanup the allocated asset data
+	for (FAssetData* AssetData : AssetDataList)
+	{
+		delete AssetData;
+	}
+
+	return bReadOk;
 }
 
 void UAssetRegistryImpl::InitializeTemporaryAssetRegistryState(FAssetRegistryState& OutState, const FAssetRegistrySerializationOptions& Options, bool bRefreshExisting, const TMap<FName, FAssetData*>& OverrideData) const
diff --git a/Engine/Source/Runtime/AssetRegistry/Private/AssetRegistry.h b/Engine/Source/Runtime/AssetRegistry/Private/AssetRegistry.h
--- a/Engine/Source/Runtime/AssetRegistry/Private/AssetRegistry.h
+++ b/Engine/Source/Runtime/AssetRegistry/Private/AssetRegistry.h
@@ -100,8 +100,8 @@ public:
 	virtual void Serialize(FArchive& Ar) override;
 	virtual void AppendState(const FAssetRegistryState& InState) override;
 	virtual uint32 GetAllocatedSize(bool bLogDetailed = false) const override;
-	virtual void LoadPackageRegistryData(FArchive& Ar, TArray<FAssetData*>& Data) const override;
-	virtual void LoadPackageRegistryData(const FString& PackageFilename, TArray<FAssetData*>& Data) const override;
+	virtual void LoadPackageRegistryData(FArchive& Ar, FLoadPackageRegistryData& InOutData) const override;
+	virtual void LoadPackageRegistryData(const FString& PackageFilename, FLoadPackageRegistryData& InOutData) const override;
 	virtual void InitializeTemporaryAssetRegistryState(FAssetRegistryState& OutState, const FAssetRegistrySerializationOptions& Options, bool bRefreshExisting = false, const TMap<FName, FAssetData*>& OverrideData = TMap<FName, FAssetData*>()) const override;
 	virtual const FAssetRegistryState* GetAssetRegistryState() const override;
 	virtual const TSet<FName>& GetCachedEmptyPackages() const override;
@@ -318,7 +318,7 @@ private:
 	 * @param PackageReader	Previsouly initialized package reader that maps to the data to be read
 	 * @param AssetDataList	List of the read asset data
 	 */
-	bool ReadAssetFile(FPackageReader& PackageReader, TArray<FAssetData*>& AssetDataList) const;
+	bool ReadAssetFile(FPackageReader& PackageReader, FLoadPackageRegistryData& InOutData) const;
 private:
 	
 	/** Internal state of the cached asset registry */
diff --git a/Engine/Source/Runtime/AssetRegistry/Public/AssetRegistry/IAssetRegistry.h b/Engine/Source/Runtime/AssetRegistry/Public/AssetRegistry/IAssetRegistry.h
--- a/Engine/Source/Runtime/AssetRegistry/Public/AssetRegistry/IAssetRegistry.h
+++ b/Engine/Source/Runtime/AssetRegistry/Public/AssetRegistry/IAssetRegistry.h
@@ -508,11 +508,23 @@ public:
 	/** Fills in FAssetRegistrySerializationOptions from ini, optionally using a target platform ini name */
 	virtual void InitializeSerializationOptions(FAssetRegistrySerializationOptions& Options, const FString& PlatformIniName = FString()) const = 0;
 
+	struct FLoadPackageRegistryData
+	{
+		FLoadPackageRegistryData(bool bInGetDependencies = false)
+			: bGetDependencies(bInGetDependencies)
+		{
+		}
+
+		TArray<FAssetData> Data;
+		TArray<FName> DataDependencies;
+		bool bGetDependencies;
+	};
+
 	/** Load FPackageRegistry data from the supplied package */
-	virtual void LoadPackageRegistryData(FArchive& Ar, TArray<FAssetData*>& Data) const = 0;
+	virtual void LoadPackageRegistryData(FArchive& Ar, FLoadPackageRegistryData& InOutData) const = 0;
 	
 	/** Load FAssetData from the specified package filename */
-	virtual void LoadPackageRegistryData(const FString& PackageFilename, TArray<FAssetData*>& Data) const = 0;
+	virtual void LoadPackageRegistryData(const FString& PackageFilename, FLoadPackageRegistryData& InOutData) const = 0;
 	
 protected:
 	// Functions specifically for calling from the asset manager
