Refactored changelist tree items to be shareable with the submit dialog (allows symmetrical display)
Added an exception to changelist validation for deleted files
Improved display of deleted & shelved objects in changelists & submit dialog - improvements related to OFPA in a subsequent change.

diff --git a/Engine/Plugins/Editor/DataValidation/Source/DataValidation/Private/DataValidationChangelist.cpp b/Engine/Plugins/Editor/DataValidation/Source/DataValidation/Private/DataValidationChangelist.cpp
--- a/Engine/Plugins/Editor/DataValidation/Source/DataValidation/Private/DataValidationChangelist.cpp
+++ b/Engine/Plugins/Editor/DataValidation/Source/DataValidation/Private/DataValidationChangelist.cpp
@@ -74,6 +74,12 @@ EDataValidationResult UDataValidationChangelist::IsDataValid(TArray<FText>& Vali
 
 	for (const FSourceControlStateRef& File : ChangelistState->GetFilesStates())
 	{
+		// We shouldn't consider dependencies of deleted files
+		if (File->IsDeleted())
+		{
+			continue;
+		}
+
 		FString PackageName;
 		if (FPackageName::TryConvertFilenameToLongPackageName(File->GetFilename(), PackageName))
 		{
diff --git a/Engine/Source/Editor/SourceControlWindows/Private/SSourceControlChangelists.cpp b/Engine/Source/Editor/SourceControlWindows/Private/SSourceControlChangelists.cpp
--- a/Engine/Source/Editor/SourceControlWindows/Private/SSourceControlChangelists.cpp
+++ b/Engine/Source/Editor/SourceControlWindows/Private/SSourceControlChangelists.cpp
@@ -9,7 +9,6 @@
 #include "Widgets/DeclarativeSyntaxSupport.h"
 #include "Widgets/Images/SImage.h"
 #include "Widgets/Layout/SScrollBorder.h"
-#include "Widgets/SOverlay.h"
 #include "Widgets/Input/SButton.h"
 #include "Widgets/Layout/SBorder.h"
 #include "Widgets/Layout/SBox.h"
@@ -35,38 +34,6 @@
 #define LOCTEXT_NAMESPACE "SourceControlChangelist"
 
 //////////////////////////////
-static TSharedRef<SWidget> GetSCCFileWidget(FSourceControlStateRef InFileState, bool bIsShelvedFile = false)
-{
-	const FSlateBrush* IconBrush = FEditorStyle::GetBrush("ContentBrowser.ColumnViewAssetIcon");
-
-	// Make icon overlays (eg, SCC and dirty status) a reasonable size in relation to the icon size (note: it is assumed this icon is square)
-	const float ICON_SCALING_FACTOR = 0.7f;
-	const float IconOverlaySize = IconBrush->ImageSize.X * ICON_SCALING_FACTOR;
-
-	return SNew(SOverlay)
-		// The actual icon
-		+ SOverlay::Slot()
-		[
-			SNew(SImage)
-			.Image(IconBrush)
-			.ColorAndOpacity_Lambda([bIsShelvedFile]() -> FSlateColor {
-					return FSlateColor(bIsShelvedFile ? FColor::Yellow : FColor::White);
-				})
-		]
-		// Source control state
-		+ SOverlay::Slot()
-		.HAlign(HAlign_Left)
-		.VAlign(VAlign_Top)
-		[
-			SNew(SBox)
-			.WidthOverride(IconOverlaySize)
-			.HeightOverride(IconOverlaySize)
-			[
-				SNew(SLayeredImage, InFileState->GetIcon())
-			]
-		];
-}
-
 struct FSCCFileDragDropOp : public FDragDropOperation
 {
 	DRAG_DROP_OPERATOR_TYPE(FSCCFileDragDropOp, FDragDropOperation);
@@ -75,172 +42,13 @@ struct FSCCFileDragDropOp : public FDragDropOperation
 
 	virtual TSharedPtr<SWidget> GetDefaultDecorator() const override
 	{
-		return GetSCCFileWidget(Files[0]);
+		return SSourceControlCommon::GetSCCFileWidget(Files[0]);
 	}
 
 	TArray<FSourceControlStateRef> Files;
 };
 
 //////////////////////////////
-
-struct FChangelistTreeItem : public IChangelistTreeItem
-{
-	FChangelistTreeItem(FSourceControlChangelistStateRef InChangelistState)
-		: ChangelistState(InChangelistState)
-	{
-		Type = IChangelistTreeItem::Changelist;
-	}
-
-	FText GetDisplayText() const
-	{
-		return ChangelistState->GetDisplayText();
-	}
-
-	FText GetDescriptionText() const
-	{
-		return ChangelistState->GetDescriptionText();
-	}
-
-	FSourceControlChangelistStateRef ChangelistState;
-};
-
-struct FShelvedChangelistTreeItem : public IChangelistTreeItem
-{
-	FShelvedChangelistTreeItem()
-	{
-		Type = IChangelistTreeItem::ShelvedChangelist;
-	}
-
-	FText GetDisplayText() const
-	{
-		return LOCTEXT("SourceControl_ShelvedFiles", "Shelved Items");
-	}
-};
-
-struct FFileTreeItem : public IChangelistTreeItem
-{
-	FFileTreeItem(FSourceControlStateRef InFileState)
-		: FileState(InFileState)
-	{
-		Type = IChangelistTreeItem::File;
-
-		// Initialize asset data first
-		FString Filename = FileState->GetFilename();
-		USourceControlHelpers::GetAssetData(Filename, Assets);
-
-		// Initialize display-related members
-		FString AssetName = LOCTEXT("SourceControl_DefaultAssetName", "None").ToString();
-		FString AssetType = LOCTEXT("SourceControl_DefaultAssetType", "Unknown").ToString();
-		FString AssetPath = Filename;
-		FColor AssetColor = FColor(		// Copied from ContentBrowserCLR.cpp
-			127 + FColor::Red.R / 2,	// Desaturate the colors a bit (GB colors were too.. much)
-			127 + FColor::Red.G / 2,
-			127 + FColor::Red.B / 2,
-			200); // Opacity
-
-		if(Assets.Num() > 0)
-		{
-			AssetPath = Assets[0].ObjectPath.ToString();
-
-			// Strip asset name from object path
-			int32 LastDot = -1;
-			if (AssetPath.FindLastChar('.', LastDot))
-			{
-				AssetPath.LeftInline(LastDot);
-			}
-
-			// Find name, asset type & color only if there is exactly one asset
-			if (Assets.Num() == 1)
-			{
-				static FName NAME_ActorLabel(TEXT("ActorLabel"));
-				if (Assets[0].FindTag(NAME_ActorLabel))
-				{
-					Assets[0].GetTagValue(NAME_ActorLabel, AssetName);
-				}
-				else
-				{
-					AssetName = Assets[0].AssetName.ToString();
-				}
-
-				AssetType = Assets[0].AssetClass.ToString();
-
-				const FAssetToolsModule& AssetToolsModule = FModuleManager::LoadModuleChecked<FAssetToolsModule>(TEXT("AssetTools"));
-				const TSharedPtr<IAssetTypeActions> AssetTypeActions = AssetToolsModule.Get().GetAssetTypeActionsForClass(Assets[0].GetClass()).Pin();
-				if (AssetTypeActions.IsValid())
-				{
-					AssetColor = AssetTypeActions->GetTypeColor();
-				}
-				else
-				{
-					AssetColor = FColor::White;
-				}
-			}
-			else
-			{
-				AssetType = LOCTEXT("SourceCOntrol_ManyAssetType", "Multiple Assets").ToString();
-				AssetColor = FColor::White;
-			}
-		}
-
-		DisplayName = FText::FromString(AssetName);
-		DisplayPath = FText::FromString(AssetPath);
-		DisplayType = FText::FromString(AssetType);
-		DisplayColor = AssetColor;
-	}
-
-	FText GetDisplayPath() const
-	{
-		return DisplayPath;
-	}
-
-	FText GetDisplayName() const
-	{
-		return DisplayName;
-	}
-
-	FText GetDisplayType() const
-	{
-		return DisplayType;
-	}
-
-	FSlateColor GetDisplayColor() const
-	{
-		return FSlateColor(DisplayColor);
-	}
-
-	const TArray<FAssetData>& GetAssetData() const
-	{
-		return Assets;
-	}
-
-public:
-	FSourceControlStateRef FileState;
-
-private:
-	TArray<FAssetData> Assets;
-
-	FText DisplayPath;
-	FText DisplayName;
-	FText DisplayType;
-	FColor DisplayColor;
-};
-
-struct FShelvedFileTreeItem : public IChangelistTreeItem
-{
-	FShelvedFileTreeItem(FSourceControlStateRef InFileState)
-		: FileState(InFileState)
-	{
-		Type = IChangelistTreeItem::ShelvedFile;
-	}
-
-	FText GetDisplayName() const
-	{
-		return FText::FromString(FileState->GetFilename());
-	}
-
-	FSourceControlStateRef FileState;
-};
-
 SSourceControlChangelistsWidget::SSourceControlChangelistsWidget()
 {
 }
@@ -1150,14 +958,16 @@ public:
 	{
 		if (ColumnName == TEXT("Change")) // eq. to name
 		{
+			const int32 LeftOffset = (TreeItem->IsShelved() ? 60 : 40);
+
 			return SNew(SHorizontalBox)
 
 			// Icon
 			+ SHorizontalBox::Slot()
 				.AutoWidth()
-				.Padding(40, 0, 4, 0)
+				.Padding(LeftOffset, 0, 4, 0)
 				[
-					GetSCCFileWidget(TreeItem->FileState)
+					SSourceControlCommon::GetSCCFileWidget(TreeItem->FileState, TreeItem->IsShelved())
 				]
 
 			+ SHorizontalBox::Slot()
@@ -1171,7 +981,8 @@ public:
 		else if (ColumnName == TEXT("Description")) // eq. to path
 		{
 			return SNew(STextBlock)
-				.Text(this, &SFileTableRow::GetDisplayPath);
+				.Text(this, &SFileTableRow::GetDisplayPath)
+				.ToolTipText(this, &SFileTableRow::GetFilename);
 		}
 		else if (ColumnName == TEXT("Type"))
 		{
@@ -1187,22 +998,27 @@ public:
 
 	FText GetDisplayName() const
 	{
-		return TreeItem->GetDisplayName();
+		return TreeItem->GetAssetName();
+	}
+
+	FText GetFilename() const
+	{
+		return TreeItem->GetFileName();
 	}
 
 	FText GetDisplayPath() const
 	{
-		return TreeItem->GetDisplayPath();
+		return TreeItem->GetAssetPath();
 	}
 
 	FText GetDisplayType() const
 	{
-		return TreeItem->GetDisplayType();
+		return TreeItem->GetAssetType();
 	}
 
 	FSlateColor GetDisplayColor() const
 	{
-		return TreeItem->GetDisplayColor();
+		return TreeItem->GetAssetTypeColor();
 	}
 
 protected:
@@ -1294,75 +1110,6 @@ private:
 	FShelvedChangelistTreeItem* TreeItem;
 };
 
-class SShelvedFileTableRow : public SMultiColumnTableRow<FChangelistTreeItemPtr>
-{
-public:
-	SLATE_BEGIN_ARGS(SShelvedFileTableRow)
-		: _TreeItemToVisualize()
-	{}
-	SLATE_ARGUMENT(FChangelistTreeItemPtr, TreeItemToVisualize)
-		SLATE_END_ARGS()
-
-public:
-	/**
-	* Construct child widgets that comprise this widget.
-	*
-	* @param InArgs Declaration from which to construct this widget.
-	*/
-	void Construct(const FArguments& InArgs, const TSharedRef<STableViewBase>& InOwner)
-	{
-		TreeItem = static_cast<FShelvedFileTreeItem*>(InArgs._TreeItemToVisualize.Get());
-
-		auto Args = FSuperRowType::FArguments();
-		FSuperRowType::Construct(Args, InOwner);
-	}
-
-	// SMultiColumnTableRow overrides
-	virtual TSharedRef<SWidget> GenerateWidgetForColumn(const FName& ColumnName) override
-	{
-		const bool bIsShelvedFile = true;
-
-		if (ColumnName == TEXT("Change")) // eq. to name
-		{
-			return SNew(SHorizontalBox)
-
-				// Icon
-				+ SHorizontalBox::Slot()
-				.AutoWidth()
-				.Padding(60, 0, 4, 0)
-				[
-					GetSCCFileWidget(TreeItem->FileState, bIsShelvedFile)
-				]
-
-			+ SHorizontalBox::Slot()
-				.AutoWidth()
-				.VAlign(VAlign_Center)
-				[
-					SNew(STextBlock)
-					.Text(LOCTEXT("SourceControl_DefaultNameForShelvedFiles", "Unavailable"))
-				];
-		}
-		else if (ColumnName == TEXT("Description"))
-		{
-			return SNew(STextBlock)
-				.Text(this, &SShelvedFileTableRow::GetDisplayName);
-		}
-		else
-		{
-			return SNullWidget::NullWidget;
-		}
-	}
-
-	FText GetDisplayName() const
-	{
-		return TreeItem->GetDisplayName();
-	}
-
-private:
-	/** The info about the widget that we are visualizing. */
-	FShelvedFileTreeItem* TreeItem;
-};
-
 TSharedRef<ITableRow> SSourceControlChangelistsWidget::OnGenerateRow(FChangelistTreeItemPtr InTreeItem, const TSharedRef<STableViewBase>& OwnerTable)
 {
 	switch (InTreeItem->GetTreeItemType())
@@ -1381,7 +1128,7 @@ TSharedRef<ITableRow> SSourceControlChangelistsWidget::OnGenerateRow(FChangelist
 			.TreeItemToVisualize(InTreeItem);
 
 	case IChangelistTreeItem::ShelvedFile:
-		return SNew(SShelvedFileTableRow, OwnerTable)
+		return SNew(SFileTableRow, OwnerTable)
 			.TreeItemToVisualize(InTreeItem);
 
 	default:
diff --git a/Engine/Source/Editor/SourceControlWindows/Private/SSourceControlChangelists.h b/Engine/Source/Editor/SourceControlWindows/Private/SSourceControlChangelists.h
--- a/Engine/Source/Editor/SourceControlWindows/Private/SSourceControlChangelists.h
+++ b/Engine/Source/Editor/SourceControlWindows/Private/SSourceControlChangelists.h
@@ -8,65 +8,7 @@
 #include "Widgets/Views/STreeView.h"
 
 #include "ISourceControlProvider.h"
-
-struct IChangelistTreeItem;
-typedef TSharedPtr<IChangelistTreeItem> FChangelistTreeItemPtr;
-typedef TSharedRef<IChangelistTreeItem> FChangelistTreeItemRef;
-
-struct IChangelistTreeItem : TSharedFromThis<IChangelistTreeItem>
-{
-	enum TreeItemType
-	{
-		Invalid,
-		Changelist,
-		File,
-		ShelvedChangelist, // container for shelved files
-		ShelvedFile
-	};
-
-	/** Get this item's parent. Can be nullptr. */
-	FChangelistTreeItemPtr GetParent() const
-	{
-		return Parent;
-	}
-
-	/** Get this item's children, if any. Although we store as weak pointers, they are guaranteed to be valid. */
-	const TArray<FChangelistTreeItemPtr>& GetChildren() const
-	{
-		return Children;
-	}
-
-	/** Returns the TreeItem's type */
-	const TreeItemType GetTreeItemType() const
-	{
-		return Type;
-	}
-
-	/** Add a child to this item */
-	void AddChild(FChangelistTreeItemRef Child)
-	{
-		Child->Parent = AsShared();
-		Children.Add(MoveTemp(Child));
-	}
-
-	/** Remove a child from this item */
-	void RemoveChild(const FChangelistTreeItemRef& Child)
-	{
-		if (Children.Remove(Child))
-		{
-			Child->Parent = nullptr;
-		}
-	}
-
-protected:
-	/** This item's parent, if any. */
-	FChangelistTreeItemPtr Parent;
-
-	/** Array of children contained underneath this item */
-	TArray<FChangelistTreeItemPtr> Children;
-
-	TreeItemType Type;
-};
+#include "SSourceControlCommon.h"
 
 class SChangelistTree : public STreeView<FChangelistTreeItemPtr>
 {
diff --git a/Engine/Source/Editor/SourceControlWindows/Private/SSourceControlCommon.cpp b/Engine/Source/Editor/SourceControlWindows/Private/SSourceControlCommon.cpp
new file mode 100644
--- /dev/null
+++ b/Engine/Source/Editor/SourceControlWindows/Private/SSourceControlCommon.cpp
@@ -0,0 +1,185 @@
+// Copyright Epic Games, Inc. All Rights Reserved.
+
+#include "SSourceControlCommon.h"
+
+#include "AssetData.h"
+#include "AssetToolsModule.h"
+#include "EditorStyleSet.h"
+#include "SourceControlHelpers.h"
+
+#include "Widgets/SOverlay.h"
+#include "Widgets/Images/SImage.h"
+#include "Widgets/Images/SLayeredImage.h"
+#include "Widgets/Layout/SBox.h"
+
+#define LOCTEXT_NAMESPACE "SourceControlChangelist"
+
+//////////////////////////////////////////////////////////////////////////
+FChangelistTreeItemPtr IChangelistTreeItem::GetParent() const
+{
+	return Parent;
+}
+
+const TArray<FChangelistTreeItemPtr>& IChangelistTreeItem::GetChildren() const
+{
+	return Children;
+}
+
+void IChangelistTreeItem::AddChild(FChangelistTreeItemRef Child)
+{
+	Child->Parent = AsShared();
+	Children.Add(MoveTemp(Child));
+}
+
+void IChangelistTreeItem::RemoveChild(const FChangelistTreeItemRef& Child)
+{
+	if (Children.Remove(Child))
+	{
+		Child->Parent = nullptr;
+	}
+}
+
+//////////////////////////////////////////////////////////////////////////
+FText FShelvedChangelistTreeItem::GetDisplayText() const
+{
+	return LOCTEXT("SourceControl_ShelvedFiles", "Shelved Items");
+}
+
+//////////////////////////////////////////////////////////////////////////
+FFileTreeItem::FFileTreeItem(FSourceControlStateRef InFileState, bool bIsShelvedFile)
+	: FileState(InFileState)
+{
+	Type = (bIsShelvedFile ? IChangelistTreeItem::ShelvedFile : IChangelistTreeItem::File);
+	CheckBoxState = ECheckBoxState::Checked;
+
+	// Initialize asset data first
+	FString Filename = FileState->GetFilename();
+
+	USourceControlHelpers::GetAssetData(Filename, Assets);
+
+	// Initialize display-related members
+	FString TempAssetName = LOCTEXT("SourceControl_DefaultAssetName", "None").ToString();
+	FString TempAssetPath = Filename;
+	FString TempAssetType = LOCTEXT("SourceControl_DefaultAssetType", "Unknown").ToString();
+	FColor TempAssetColor = FColor(		// Copied from ContentBrowserCLR.cpp
+		127 + FColor::Red.R / 2,	// Desaturate the colors a bit (GB colors were too.. much)
+		127 + FColor::Red.G / 2,
+		127 + FColor::Red.B / 2,
+		200); // Opacity
+
+	FString TempPackageName = Filename;
+	if (!FPackageName::TryConvertFilenameToLongPackageName(Filename, TempPackageName))
+	{
+		TempPackageName = Filename;
+	}
+
+	if (Assets.Num() > 0)
+	{
+		TempAssetPath = Assets[0].ObjectPath.ToString();
+
+		// Strip asset name from object path
+		int32 LastDot = -1;
+		if (TempAssetPath.FindLastChar('.', LastDot))
+		{
+			TempAssetPath.LeftInline(LastDot);
+		}
+
+		// Find name, asset type & color only if there is exactly one asset
+		if (Assets.Num() == 1)
+		{
+			static FName NAME_ActorLabel(TEXT("ActorLabel"));
+			if (Assets[0].FindTag(NAME_ActorLabel))
+			{
+				Assets[0].GetTagValue(NAME_ActorLabel, TempAssetName);
+			}
+			else
+			{
+				TempAssetName = Assets[0].AssetName.ToString();
+			}
+
+			TempAssetType = Assets[0].AssetClass.ToString();
+
+			const FAssetToolsModule& AssetToolsModule = FModuleManager::LoadModuleChecked<FAssetToolsModule>(TEXT("AssetTools"));
+			const TSharedPtr<IAssetTypeActions> AssetTypeActions = AssetToolsModule.Get().GetAssetTypeActionsForClass(Assets[0].GetClass()).Pin();
+			if (AssetTypeActions.IsValid())
+			{
+				TempAssetColor = AssetTypeActions->GetTypeColor();
+			}
+			else
+			{
+				TempAssetColor = FColor::White;
+			}
+		}
+		else
+		{
+			TempAssetType = LOCTEXT("SourceCOntrol_ManyAssetType", "Multiple Assets").ToString();
+			TempAssetColor = FColor::White;
+		}
+
+		// Beautify the package name
+		TempPackageName = TempAssetPath + "." + TempAssetName;
+	}
+	else if (FPackageName::TryConvertFilenameToLongPackageName(Filename, TempPackageName))
+	{
+		// Fake asset name, asset path from the package name
+		TempAssetPath = TempPackageName;
+
+		int32 LastSlash = -1;
+		if (TempPackageName.FindLastChar('/', LastSlash))
+		{
+			TempAssetName = TempPackageName;
+			TempAssetName.RightChopInline(LastSlash + 1);
+		}
+	}
+	else
+	{
+		TempPackageName = Filename; // put back original package name if the try failed
+	}
+
+	// Finally, assign the temp variables to the member variables
+	AssetName = FText::FromString(TempAssetName);
+	AssetPath = FText::FromString(TempAssetPath);
+	AssetType = FText::FromString(TempAssetType);
+	AssetTypeColor = TempAssetColor;
+	PackageName = FText::FromString(TempPackageName);
+}
+
+//////////////////////////////////////////////////////////////////////////
+namespace SSourceControlCommon
+{
+
+TSharedRef<SWidget> GetSCCFileWidget(FSourceControlStateRef InFileState, bool bIsShelvedFile)
+{
+	const FSlateBrush* IconBrush = FEditorStyle::GetBrush("ContentBrowser.ColumnViewAssetIcon");
+
+	// Make icon overlays (eg, SCC and dirty status) a reasonable size in relation to the icon size (note: it is assumed this icon is square)
+	const float ICON_SCALING_FACTOR = 0.7f;
+	const float IconOverlaySize = IconBrush->ImageSize.X * ICON_SCALING_FACTOR;
+
+	return SNew(SOverlay)
+		// The actual icon
+		+ SOverlay::Slot()
+		[
+			SNew(SImage)
+			.Image(IconBrush)
+			.ColorAndOpacity_Lambda([bIsShelvedFile]() -> FSlateColor {
+				return FSlateColor(bIsShelvedFile ? FColor::Yellow : FColor::White);
+			})
+		]
+		// Source control state
+		+ SOverlay::Slot()
+		.HAlign(HAlign_Left)
+		.VAlign(VAlign_Top)
+		[
+			SNew(SBox)
+			.WidthOverride(IconOverlaySize)
+			.HeightOverride(IconOverlaySize)
+			[
+				SNew(SLayeredImage, InFileState->GetIcon())
+			]
+		];
+}
+
+} // end of namespace SSourceControlCommon
+
+#undef LOCTEXT_NAMESPACE
\ No newline at end of file
diff --git a/Engine/Source/Editor/SourceControlWindows/Private/SSourceControlCommon.h b/Engine/Source/Editor/SourceControlWindows/Private/SSourceControlCommon.h
new file mode 100644
--- /dev/null
+++ b/Engine/Source/Editor/SourceControlWindows/Private/SSourceControlCommon.h
@@ -0,0 +1,175 @@
+// Copyright Epic Games, Inc. All Rights Reserved.
+
+#pragma once
+
+#include "CoreMinimal.h"
+#include "Styling/SlateColor.h"
+#include "Styling/SlateTypes.h"
+#include "ISourceControlProvider.h"
+
+struct FAssetData;
+struct IChangelistTreeItem;
+typedef TSharedPtr<IChangelistTreeItem> FChangelistTreeItemPtr;
+typedef TSharedRef<IChangelistTreeItem> FChangelistTreeItemRef;
+
+struct IChangelistTreeItem : TSharedFromThis<IChangelistTreeItem>
+{
+	enum TreeItemType
+	{
+		Invalid,
+		Changelist,
+		File,
+		ShelvedChangelist, // container for shelved files
+		ShelvedFile
+	};
+
+	/** Get this item's parent. Can be nullptr. */
+	FChangelistTreeItemPtr GetParent() const;
+
+	/** Get this item's children, if any. Although we store as weak pointers, they are guaranteed to be valid. */
+	const TArray<FChangelistTreeItemPtr>& GetChildren() const;
+
+	/** Returns the TreeItem's type */
+	const TreeItemType GetTreeItemType() const
+	{
+		return Type;
+	}
+
+	/** Add a child to this item */
+	void AddChild(FChangelistTreeItemRef Child);
+
+	/** Remove a child from this item */
+	void RemoveChild(const FChangelistTreeItemRef& Child);
+
+protected:
+	/** This item's parent, if any. */
+	FChangelistTreeItemPtr Parent;
+
+	/** Array of children contained underneath this item */
+	TArray<FChangelistTreeItemPtr> Children;
+
+	TreeItemType Type;
+};
+
+struct FChangelistTreeItem : public IChangelistTreeItem
+{
+	FChangelistTreeItem(FSourceControlChangelistStateRef InChangelistState)
+		: ChangelistState(InChangelistState)
+	{
+		Type = IChangelistTreeItem::Changelist;
+	}
+
+	FText GetDisplayText() const
+	{
+		return ChangelistState->GetDisplayText();
+	}
+
+	FText GetDescriptionText() const
+	{
+		return ChangelistState->GetDescriptionText();
+	}
+
+	FSourceControlChangelistStateRef ChangelistState;
+};
+
+struct FShelvedChangelistTreeItem : public IChangelistTreeItem
+{
+	FShelvedChangelistTreeItem()
+	{
+		Type = IChangelistTreeItem::ShelvedChangelist;
+	}
+
+	FText GetDisplayText() const;
+};
+
+struct FFileTreeItem : public IChangelistTreeItem
+{
+	explicit FFileTreeItem(FSourceControlStateRef InFileState, bool bIsShelvedFile = false);
+
+	/** Returns the asset name of the item */
+	FText GetAssetName() const { return AssetName; }
+
+	/** Returns the asset path of the item */
+	FText GetAssetPath() const { return AssetPath; }
+
+	/** Returns the asset type of the item */
+	FText GetAssetType() const { return AssetType; }
+
+	/** Returns the asset type color of the item */
+	FSlateColor GetAssetTypeColor() const { return FSlateColor(AssetTypeColor); }
+
+	/** Returns the package name of the item to display */
+	FText GetPackageName() const { return PackageName; }
+
+	/** Returns the file name of the item in source control */
+	FText GetFileName() const { return FText::FromString(FileState->GetFilename()); }
+
+	/** Returns the name of the icon to be used in the list item widget */
+	FName GetIconName() const { return FileState->GetIcon().GetStyleName(); }
+
+	/** Returns the tooltip text for the icon */
+	FText GetIconTooltip() const { return FileState->GetDisplayTooltip(); }
+
+	/** Returns the checkbox state of this item */
+	ECheckBoxState GetCheckBoxState() const { return CheckBoxState; }
+
+	/** Sets the checkbox state of this item */
+	void SetCheckBoxState(ECheckBoxState NewState) { CheckBoxState = NewState; }
+
+	/** true if the item is not in source control and needs to be added prior to checkin */
+	bool NeedsAdding() const { return !FileState->IsSourceControlled(); }
+
+	/** true if the item is in source control and is able to be checked in */
+	bool CanCheckIn() const { return FileState->CanCheckIn() || FileState->IsDeleted(); }
+
+	/** true if the item is enabled in the list */
+	bool IsEnabled() const { return !FileState->IsConflicted() && FileState->IsCurrent(); }
+
+	/** true if the item is source controlled and not marked for add nor for delete */
+	bool CanDiff() const { return FileState->IsSourceControlled() && !FileState->IsAdded() && !FileState->IsDeleted(); }
+
+	const TArray<FAssetData>& GetAssetData() const
+	{
+		return Assets;
+	}
+
+	bool IsShelved() const { return GetTreeItemType() == IChangelistTreeItem::ShelvedFile; }
+
+public:
+	/** Shared pointer to the source control state object itself */
+	FSourceControlStateRef FileState;
+
+private:
+	/** Checkbox state, used only in the Submit dialog */
+	ECheckBoxState CheckBoxState;
+
+	/** Cached asset name to display */
+	FText AssetName;
+
+	/** Cached asset path to display */
+	FText AssetPath;
+
+	/** Cached asset type to display */
+	FText AssetType;
+
+	/** Cached asset type related color to display */
+	FColor AssetTypeColor;
+
+	/** Cached package name to display */
+	FText PackageName;
+
+	/** Matching asset(s) to facilitate Locate in content browser */
+	TArray<FAssetData> Assets;
+};
+
+struct FShelvedFileTreeItem : public FFileTreeItem
+{
+	explicit FShelvedFileTreeItem(FSourceControlStateRef InFileState)
+		: FFileTreeItem(InFileState, /*bIsShelved=*/true)
+	{}
+};
+
+namespace SSourceControlCommon
+{
+	TSharedRef<SWidget> GetSCCFileWidget(FSourceControlStateRef InFileState, bool bIsShelvedFile = false);
+}
\ No newline at end of file
diff --git a/Engine/Source/Editor/SourceControlWindows/Private/SSourceControlSubmit.cpp b/Engine/Source/Editor/SourceControlWindows/Private/SSourceControlSubmit.cpp
--- a/Engine/Source/Editor/SourceControlWindows/Private/SSourceControlSubmit.cpp
+++ b/Engine/Source/Editor/SourceControlWindows/Private/SSourceControlSubmit.cpp
@@ -41,40 +41,12 @@ namespace SSourceControlSubmitWidgetDefs
 }
 
 
-FSubmitItem::FSubmitItem(const FSourceControlStateRef& InItem)
-	: Item(InItem)
-{
-	CheckBoxState = ECheckBoxState::Checked;
-
-	AssetName = FText::FromString(TEXT("None"));
-	PackageName = FText::FromString(Item->GetFilename());
-	FileName = PackageName;
-
-	FString LongPackageName;
-	if (FPackageName::TryConvertFilenameToLongPackageName(InItem->GetFilename(), LongPackageName))
-	{
-		PackageName = FText::FromString(LongPackageName);
-
-		TArray<FAssetData> Assets;
-		FAssetRegistryModule& AssetRegistryModule = FModuleManager::LoadModuleChecked<FAssetRegistryModule>(TEXT("AssetRegistry"));
-		AssetRegistryModule.Get().GetAssetsByPackageName(*LongPackageName, Assets);
-
-		if (Assets.Num())
-		{
-			const FAssetData& AssetData = Assets[0];
-			PackageName = FText::FromString(AssetData.PackageName.ToString());
-			AssetName = FText::FromString(AssetData.AssetName.ToString());
-		}
-	}
-}
-
-
 void SSourceControlSubmitListRow::Construct(const FArguments& InArgs, const TSharedRef<STableViewBase>& InOwnerTableView)
 {
 	SourceControlSubmitWidgetPtr = InArgs._SourceControlSubmitWidget;
 	Item = InArgs._Item;
 
-	SMultiColumnTableRow<TSharedPtr<FSubmitItem>>::Construct(FSuperRowType::FArguments(), InOwnerTableView);
+	SMultiColumnTableRow<TSharedPtr<FFileTreeItem>>::Construct(FSuperRowType::FArguments(), InOwnerTableView);
 }
 
 
@@ -111,7 +83,7 @@ void SSourceControlSubmitWidget::Construct(const FArguments& InArgs)
 
 	for (const auto& Item : InArgs._Items.Get())
 	{
-		ListViewItems.Add(MakeShareable(new FSubmitItem(Item)));
+		ListViewItems.Add(MakeShareable(new FFileTreeItem(Item)));
 	}
 
 	TSharedRef<SHeaderRow> HeaderRowWidget = SNew(SHeaderRow);
@@ -195,7 +167,7 @@ void SSourceControlSubmitWidget::Construct(const FArguments& InArgs)
 	[
 		SNew(SBorder)
 		[
-			SAssignNew(ListView, SListView<TSharedPtr<FSubmitItem>>)
+			SAssignNew(ListView, SListView<TSharedPtr<FFileTreeItem>>)
 			.ItemHeight(20)
 			.ListItemsSource(&ListViewItems)
 			.OnGenerateRow(this, &SSourceControlSubmitWidget::OnGenerateRowForList)
@@ -360,7 +332,7 @@ void SSourceControlSubmitWidget::OnDiffAgainstDepot()
 	}
 }
 
-void SSourceControlSubmitWidget::OnDiffAgainstDepotSelected(TSharedPtr<FSubmitItem> InSelectedItem)
+void SSourceControlSubmitWidget::OnDiffAgainstDepotSelected(TSharedPtr<FFileTreeItem> InSelectedItem)
 {
 	FString PackageName;
 	if (FPackageName::TryConvertFilenameToLongPackageName(InSelectedItem->GetFileName().ToString(), PackageName))
@@ -393,7 +365,7 @@ FReply SSourceControlSubmitWidget::OnKeyDown( const FGeometry& MyGeometry, const
    return FReply::Unhandled();
 }
 
-TSharedRef<SWidget> SSourceControlSubmitWidget::GenerateWidgetForItemAndColumn(TSharedPtr<FSubmitItem> Item, const FName ColumnID) const
+TSharedRef<SWidget> SSourceControlSubmitWidget::GenerateWidgetForItemAndColumn(TSharedPtr<FFileTreeItem> Item, const FName ColumnID) const
 {
 	check(Item.IsValid());
 
@@ -408,8 +380,8 @@ TSharedRef<SWidget> SSourceControlSubmitWidget::GenerateWidgetForItemAndColumn(T
 			.Padding(RowPadding)
 			[
 				SNew(SCheckBox)
-				.IsChecked(Item.Get(), &FSubmitItem::GetCheckBoxState)
-				.OnCheckStateChanged(Item.Get(), &FSubmitItem::SetCheckBoxState)
+				.IsChecked(Item.Get(), &FFileTreeItem::GetCheckBoxState)
+				.OnCheckStateChanged(Item.Get(), &FFileTreeItem::SetCheckBoxState)
 			];
 	}
 	else if (ColumnID == SSourceControlSubmitWidgetDefs::ColumnID_IconLabel)
@@ -565,7 +537,7 @@ bool SSourceControlSubmitWidget::CanCheckOut() const
 }
 
 
-TSharedRef<ITableRow> SSourceControlSubmitWidget::OnGenerateRowForList(TSharedPtr<FSubmitItem> SubmitItem, const TSharedRef<STableViewBase>& OwnerTable)
+TSharedRef<ITableRow> SSourceControlSubmitWidget::OnGenerateRowForList(TSharedPtr<FFileTreeItem> SubmitItem, const TSharedRef<STableViewBase>& OwnerTable)
 {
 	TSharedRef<ITableRow> Row =
 	SNew(SSourceControlSubmitListRow, OwnerTable)
@@ -612,12 +584,12 @@ void SSourceControlSubmitWidget::SortTree()
 	{
 		if (SortMode == EColumnSortMode::Ascending)
 		{
-			ListViewItems.Sort([](const TSharedPtr<FSubmitItem>& A, const TSharedPtr<FSubmitItem>& B) {
+			ListViewItems.Sort([](const TSharedPtr<FFileTreeItem>& A, const TSharedPtr<FFileTreeItem>& B) {
 				return A->GetAssetName().ToString() < B->GetAssetName().ToString(); });
 		}
 		else if (SortMode == EColumnSortMode::Descending)
 		{
-			ListViewItems.Sort([](const TSharedPtr<FSubmitItem>& A, const TSharedPtr<FSubmitItem>& B) {
+			ListViewItems.Sort([](const TSharedPtr<FFileTreeItem>& A, const TSharedPtr<FFileTreeItem>& B) {
 				return A->GetAssetName().ToString() >= B->GetAssetName().ToString(); });
 		}
 	}
@@ -625,12 +597,12 @@ void SSourceControlSubmitWidget::SortTree()
 	{
 		if (SortMode == EColumnSortMode::Ascending)
 		{
-			ListViewItems.Sort([](const TSharedPtr<FSubmitItem>& A, const TSharedPtr<FSubmitItem>& B) {
+			ListViewItems.Sort([](const TSharedPtr<FFileTreeItem>& A, const TSharedPtr<FFileTreeItem>& B) {
 				return A->GetPackageName().ToString() < B->GetPackageName().ToString(); });
 		}
 		else if (SortMode == EColumnSortMode::Descending)
 		{
-			ListViewItems.Sort([](const TSharedPtr<FSubmitItem>& A, const TSharedPtr<FSubmitItem>& B) {
+			ListViewItems.Sort([](const TSharedPtr<FFileTreeItem>& A, const TSharedPtr<FFileTreeItem>& B) {
 				return A->GetPackageName().ToString() >= B->GetPackageName().ToString(); });
 		}
 	}
@@ -638,12 +610,12 @@ void SSourceControlSubmitWidget::SortTree()
 	{
 		if (SortMode == EColumnSortMode::Ascending)
 		{
-			ListViewItems.Sort([](const TSharedPtr<FSubmitItem>& A, const TSharedPtr<FSubmitItem>& B) {
+			ListViewItems.Sort([](const TSharedPtr<FFileTreeItem>& A, const TSharedPtr<FFileTreeItem>& B) {
 				return A->GetIconName().ToString() < B->GetIconName().ToString(); });
 		}
 		else if (SortMode == EColumnSortMode::Descending)
 		{
-			ListViewItems.Sort([](const TSharedPtr<FSubmitItem>& A, const TSharedPtr<FSubmitItem>& B) {
+			ListViewItems.Sort([](const TSharedPtr<FFileTreeItem>& A, const TSharedPtr<FFileTreeItem>& B) {
 				return A->GetIconName().ToString() >= B->GetIconName().ToString(); });
 		}
 	}
diff --git a/Engine/Source/Editor/SourceControlWindows/Private/SSourceControlSubmit.h b/Engine/Source/Editor/SourceControlWindows/Private/SSourceControlSubmit.h
--- a/Engine/Source/Editor/SourceControlWindows/Private/SSourceControlSubmit.h
+++ b/Engine/Source/Editor/SourceControlWindows/Private/SSourceControlSubmit.h
@@ -12,6 +12,7 @@
 #include "Widgets/Views/STableViewBase.h"
 #include "Widgets/Views/STableRow.h"
 #include "ISourceControlState.h"
+#include "SSourceControlCommon.h"
 
 class SMultiLineEditableTextBox;
 class SWindow;
@@ -39,62 +40,6 @@ public:
 	FText Description;
 };
 
-class FSubmitItem : public TSharedFromThis<FSubmitItem>
-{
-public:
-	/** Constructor */
-	explicit FSubmitItem(const FSourceControlStateRef& InItem);
-
-	/** Returns the asset name of the item */
-	FText GetAssetName() const { return AssetName; }
-
-	/** Returns the package name of the item to display */
-	FText GetPackageName() const { return PackageName; }
-
-	/** Returns the file name of the item in source control */
-	FText GetFileName() const { return FileName; }
-
-	/** Returns the name of the icon to be used in the list item widget */
-	FName GetIconName() const { return Item->GetIcon().GetStyleName(); }
-
-	/** Returns the tooltip text for the icon */
-	FText GetIconTooltip() const { return Item->GetDisplayTooltip(); }
-
-	/** Returns the checkbox state of this item */
-	ECheckBoxState GetCheckBoxState() const { return CheckBoxState; }
-
-	/** Sets the checkbox state of this item */
-	void SetCheckBoxState(ECheckBoxState NewState) { CheckBoxState = NewState; }
-
-	/** true if the item is not in source control and needs to be added prior to checkin */
-	bool NeedsAdding() const { return !Item->IsSourceControlled(); }
-
-	/** true if the item is in source control and is able to be checked in */
-	bool CanCheckIn() const { return Item->CanCheckIn() || Item->IsDeleted(); }
-
-	/** true if the item is enabled in the list */
-	bool IsEnabled() const { return !Item->IsConflicted() && Item->IsCurrent(); }
-
-	/** true if the item is source controlled and not marked for add nor for delete */
-	bool CanDiff() const { return Item->IsSourceControlled() && !Item->IsAdded() && !Item->IsDeleted(); }
-
-private:
-	/** Shared pointer to the source control state object itself */
-	FSourceControlStateRef Item;
-
-	/** Checkbox state */
-	ECheckBoxState CheckBoxState;
-
-	/** Cached asset name to display in the listview */
-	FText AssetName;
-
-	/** Cached package name to display in the listview */
-	FText PackageName;
-
-	/** Cached file name on disk for source control */
-	FText FileName;
-};
-
 class SSourceControlSubmitWidget : public SCompoundWidget
 {
 public:
@@ -132,7 +77,7 @@ public:
 	ESubmitResults::Type GetResult() { return DialogResult; }
 
 	/** Returns a widget representing the item and column supplied */
-	TSharedRef<SWidget> GenerateWidgetForItemAndColumn(TSharedPtr<FSubmitItem> Item, const FName ColumnID) const;
+	TSharedRef<SWidget> GenerateWidgetForItemAndColumn(TSharedPtr<FFileTreeItem> Item, const FName ColumnID) const;
 
 	/** Gets the requested files and the change list description*/
 	void FillChangeListDescription(FChangeListDescription& OutDesc);
@@ -178,7 +123,7 @@ private:
 	bool CanCheckOut() const;
 
 	/** Called by SListView to get a widget corresponding to the supplied item */
-	TSharedRef<ITableRow> OnGenerateRowForList(TSharedPtr<FSubmitItem> SubmitItemData, const TSharedRef<STableViewBase>& OwnerTable);
+	TSharedRef<ITableRow> OnGenerateRowForList(TSharedPtr<FFileTreeItem> SubmitItemData, const TSharedRef<STableViewBase>& OwnerTable);
 
 	/**
 	 * Returns the current column sort mode (ascending or descending) if the ColumnId parameter matches the current
@@ -213,16 +158,16 @@ private:
 
 	bool CanDiffAgainstDepot() const;
 	void OnDiffAgainstDepot();
-	void OnDiffAgainstDepotSelected(TSharedPtr<FSubmitItem> InSelectedItem);
+	void OnDiffAgainstDepotSelected(TSharedPtr<FFileTreeItem> InSelectedItem);
 	
 private:
 	ESubmitResults::Type DialogResult;
 
 	/** ListBox for selecting which object to consolidate */
-	TSharedPtr<SListView<TSharedPtr<FSubmitItem>>> ListView;
+	TSharedPtr<SListView<TSharedPtr<FFileTreeItem>>> ListView;
 
 	/** Collection of objects (Widgets) to display in the List View. */
-	TArray<TSharedPtr<FSubmitItem>> ListViewItems;
+	TArray<TSharedPtr<FFileTreeItem>> ListViewItems;
 
 	/** Pointer to the parent modal window */
 	TWeakPtr<SWindow> ParentFrame;
@@ -243,7 +188,7 @@ private:
 	static FText SavedChangeListDescription;	
 };
 
-class SSourceControlSubmitListRow : public SMultiColumnTableRow<TSharedPtr<FSubmitItem>>
+class SSourceControlSubmitListRow : public SMultiColumnTableRow<TSharedPtr<FFileTreeItem>>
 {
 public:
 
@@ -253,7 +198,7 @@ public:
 	SLATE_ARGUMENT(TSharedPtr<SSourceControlSubmitWidget>, SourceControlSubmitWidget)
 
 		/** The list item for this row */
-		SLATE_ARGUMENT(TSharedPtr<FSubmitItem>, Item)
+		SLATE_ARGUMENT(TSharedPtr<FFileTreeItem>, Item)
 
 	SLATE_END_ARGS()
 
@@ -269,5 +214,5 @@ private:
 	TWeakPtr<SSourceControlSubmitWidget> SourceControlSubmitWidgetPtr;
 
 	/** The item associated with this row of data */
-	TSharedPtr<FSubmitItem> Item;
+	TSharedPtr<FFileTreeItem> Item;
 };
