Minor UX improvements to the Source Control Changelists window

diff --git a/Engine/Source/Editor/EditorStyle/Private/StarshipStyle.cpp b/Engine/Source/Editor/EditorStyle/Private/StarshipStyle.cpp
--- a/Engine/Source/Editor/EditorStyle/Private/StarshipStyle.cpp
+++ b/Engine/Source/Editor/EditorStyle/Private/StarshipStyle.cpp
@@ -7261,6 +7261,7 @@ void FStarshipEditorStyle::FStyle::SetupSourceControlStyles()
 
 		Set("SourceControl.ChangelistsTab", new CORE_IMAGE_BRUSH_SVG("Starship/Common/check-circle", Icon16x16));
 		Set("SourceControl.Changelist", new CORE_IMAGE_BRUSH_SVG("Starship/SourceControl/SCC_Changelist", Icon16x16, FStyleColors::AccentRed));
+		Set("SourceControl.ShelvedChangelist", new CORE_IMAGE_BRUSH_SVG("Starship/SourceControl/SCC_Changelist", Icon16x16, FStyleColors::AccentBlue));
 
 		Set("SourceControl.Add", new CORE_IMAGE_BRUSH_SVG( "Starship/SourceControl/SCC_ContentAdd",Icon16x16));
 		Set("SourceControl.Edit", new CORE_IMAGE_BRUSH_SVG( "Starship/SourceControl/SCC_CheckedOut", Icon16x16));
diff --git a/Engine/Source/Editor/SourceControlWindows/Private/SSourceControlChangelists.cpp b/Engine/Source/Editor/SourceControlWindows/Private/SSourceControlChangelists.cpp
--- a/Engine/Source/Editor/SourceControlWindows/Private/SSourceControlChangelists.cpp
+++ b/Engine/Source/Editor/SourceControlWindows/Private/SSourceControlChangelists.cpp
@@ -723,7 +723,7 @@ TSharedPtr<SWidget> SSourceControlChangelistsWidget::OnOpenContextMenu()
 	// This should appear only on change lists
 	if (bHasSelectedChangelist)
 	{
-		Section.AddMenuEntry("SubmitChangelist", LOCTEXT("SourceControl_SubmitChangelist", "Submit Changelist"), LOCTEXT("SourceControl_SubmitChangeslit_Tooltip", "Submits a changelist"), FSlateIcon(),
+		Section.AddMenuEntry("SubmitChangelist", LOCTEXT("SourceControl_SubmitChangelist", "Submit Changelist..."), LOCTEXT("SourceControl_SubmitChangeslit_Tooltip", "Submits a changelist"), FSlateIcon(),
 			FUIAction(
 				FExecuteAction::CreateSP(this, &SSourceControlChangelistsWidget::OnSubmitChangelist),
 				FCanExecuteAction::CreateSP(this, &SSourceControlChangelistsWidget::CanSubmitChangelist)));
@@ -767,7 +767,7 @@ TSharedPtr<SWidget> SSourceControlChangelistsWidget::OnOpenContextMenu()
 	if (bHasSelectedShelvedFiles)
 	{
 		// Diff against workspace
-		Section.AddMenuEntry("DiffAgainstWorkspace", LOCTEXT("SourceControl_DiffAgainstWorkspace", "Diff Against Workspace Files"), LOCTEXT("SourceControl_DiffAgainstWorkspace_Tooltip", "Diff shelved file against the (local) workspace file"), FSlateIcon(),
+		Section.AddMenuEntry("DiffAgainstWorkspace", LOCTEXT("SourceControl_DiffAgainstWorkspace", "Diff Against Workspace Files..."), LOCTEXT("SourceControl_DiffAgainstWorkspace_Tooltip", "Diff shelved file against the (local) workspace file"), FSlateIcon(),
 			FUIAction(
 				FExecuteAction::CreateSP(this, &SSourceControlChangelistsWidget::OnDiffAgainstWorkspace),
 				FCanExecuteAction::CreateSP(this, &SSourceControlChangelistsWidget::CanDiffAgainstWorkspace)));
@@ -781,13 +781,13 @@ TSharedPtr<SWidget> SSourceControlChangelistsWidget::OnOpenContextMenu()
 	// This should appear only if we have no selection
 	if (bHasEmptySelection)
 	{
-		Section.AddMenuEntry("NewChangelist", LOCTEXT("SourceControl_NewChangelist", "New Changelist"), LOCTEXT("SourceControl_NewChangelist_Tooltip", "Creates an empty changelist"), FSlateIcon(),
+		Section.AddMenuEntry("NewChangelist", LOCTEXT("SourceControl_NewChangelist", "New Changelist..."), LOCTEXT("SourceControl_NewChangelist_Tooltip", "Creates an empty changelist"), FSlateIcon(),
 			FUIAction(FExecuteAction::CreateSP(this, &SSourceControlChangelistsWidget::OnNewChangelist)));
 	}
 
 	if (bHasSelectedChangelist)
 	{
-		Section.AddMenuEntry("EditChangelist", LOCTEXT("SourceControl_EditChangelist", "Edit Changelist"), LOCTEXT("SourceControl_Edit_Changelist_Tooltip", "Edit a changelist description"), FSlateIcon(),
+		Section.AddMenuEntry("EditChangelist", LOCTEXT("SourceControl_EditChangelist", "Edit Changelist..."), LOCTEXT("SourceControl_Edit_Changelist_Tooltip", "Edit a changelist description"), FSlateIcon(),
 			FUIAction(FExecuteAction::CreateSP(this, &SSourceControlChangelistsWidget::OnEditChangelist)));
 
 		Section.AddMenuEntry("DeleteChangelist", LOCTEXT("SourceControl_DeleteChangelist", "Delete Empty Changelist"), LOCTEXT("SourceControl_Delete_Changelist_Tooltip", "Deletes an empty changelist"), FSlateIcon(),
@@ -801,16 +801,16 @@ TSharedPtr<SWidget> SSourceControlChangelistsWidget::OnOpenContextMenu()
 	{
 		Section.AddSeparator("FilesSeparator");
 
-		Section.AddMenuEntry("LocateFile", LOCTEXT("SourceControl_LocateFile", "Locate File"), LOCTEXT("SourceControl_LocateFile_Tooltip", "Locate File in Project..."), FSlateIcon(),
+		Section.AddMenuEntry("LocateFile", LOCTEXT("SourceControl_LocateFile", "Locate File..."), LOCTEXT("SourceControl_LocateFile_Tooltip", "Locate File in Project..."), FSlateIcon(),
 			FUIAction(
 				FExecuteAction::CreateSP(this, &SSourceControlChangelistsWidget::OnLocateFile),
 				FCanExecuteAction::CreateSP(this, &SSourceControlChangelistsWidget::CanLocateFile)));
 
-		Section.AddMenuEntry("ShowHistory", LOCTEXT("SourceControl_ShowHistory", "Show History"), LOCTEXT("SourceControl_ShowHistory_ToolTip", "Show File History From Selection..."), FSlateIcon(),
+		Section.AddMenuEntry("ShowHistory", LOCTEXT("SourceControl_ShowHistory", "Show History..."), LOCTEXT("SourceControl_ShowHistory_ToolTip", "Show File History From Selection..."), FSlateIcon(),
 			FUIAction(
 				FExecuteAction::CreateSP(this, &SSourceControlChangelistsWidget::OnShowHistory)));
 
-		Section.AddMenuEntry("DiffAgainstLocalVersion", LOCTEXT("SourceControl_DiffAgainstDepot", "Diff Against Depot"), LOCTEXT("SourceControl_DiffAgainstLocal_Tooltip", "Diff local file against depot revision."), FSlateIcon(),
+		Section.AddMenuEntry("DiffAgainstLocalVersion", LOCTEXT("SourceControl_DiffAgainstDepot", "Diff Against Depot..."), LOCTEXT("SourceControl_DiffAgainstLocal_Tooltip", "Diff local file against depot revision."), FSlateIcon(),
 			FUIAction(
 				FExecuteAction::CreateSP(this, &SSourceControlChangelistsWidget::OnDiffAgainstDepot),
 				FCanExecuteAction::CreateSP(this, &SSourceControlChangelistsWidget::CanDiffAgainstDepot)));
@@ -1114,7 +1114,7 @@ public:
 					.Padding(5, 0, 0, 0)
 					[
 						SNew(SImage)
-						.Image(FEditorStyle::GetBrush("SourceControl.Changelist"))
+						.Image(FEditorStyle::GetBrush("SourceControl.ShelvedChangelist"))
 					]
 				+ SHorizontalBox::Slot()
 					.Padding(2.0f, 0.0f)
diff --git a/Engine/Source/Editor/SourceControlWindows/Private/SSourceControlCommon.cpp b/Engine/Source/Editor/SourceControlWindows/Private/SSourceControlCommon.cpp
--- a/Engine/Source/Editor/SourceControlWindows/Private/SSourceControlCommon.cpp
+++ b/Engine/Source/Editor/SourceControlWindows/Private/SSourceControlCommon.cpp
@@ -152,7 +152,9 @@ FFileTreeItem::FFileTreeItem(FSourceControlStateRef InFileState, bool bIsShelved
 	}
 	else
 	{
+		TempAssetName = FPaths::GetCleanFilename(Filename);
 		TempPackageName = Filename; // put back original package name if the try failed
+		TempAssetType = FText::Format(LOCTEXT("SourceControl_FileTypeDefault", "{0} File"), FText::FromString(FPaths::GetExtension(Filename).ToUpper())).ToString();
 	}
 
 	// Finally, assign the temp variables to the member variables
