Source control changelists - Disable submit button if the validation contains errors

diff --git a/Engine/Source/Editor/SourceControlWindows/Private/SSourceControlChangelists.cpp b/Engine/Source/Editor/SourceControlWindows/Private/SSourceControlChangelists.cpp
--- a/Engine/Source/Editor/SourceControlWindows/Private/SSourceControlChangelists.cpp
+++ b/Engine/Source/Editor/SourceControlWindows/Private/SSourceControlChangelists.cpp
@@ -33,6 +33,8 @@
 
 #define LOCTEXT_NAMESPACE "SourceControlChangelist"
 
+DEFINE_LOG_CATEGORY_STATIC(LogSourceControlChangelist, All, All);
+
 //////////////////////////////
 struct FSCCFileDragDropOp : public FDragDropOperation
 {
@@ -517,7 +519,7 @@ void SSourceControlChangelistsWidget::OnDeleteShelvedFiles()
 	SourceControlProvider.Execute(DeleteShelvedOperation, GetChangelistFromSelection(), GetSelectedShelvedFiles());
 }
 
-static void GetChangelistValidationResult(FSourceControlChangelistPtr InChangelist, FString& OutValidationText, FName& OutValidationIcon)
+static bool GetChangelistValidationResult(FSourceControlChangelistPtr InChangelist, FString& OutValidationText, FName& OutValidationIcon)
 {
 	FSourceControlPreSubmitDataValidationDelegate ValidationDelegate = ISourceControlModule::Get().GetRegisteredPreSubmitDataValidation();
 
@@ -525,12 +527,15 @@ static void GetChangelistValidationResult(FSourceControlChangelistPtr InChangeli
 	TArray<FText> ValidationErrors;
 	TArray<FText> ValidationWarnings;
 
+	bool bValidationResult = true;
+
 	if (ValidationDelegate.ExecuteIfBound(InChangelist, ValidationResult, ValidationErrors, ValidationWarnings))
 	{
 		if (ValidationResult == EDataValidationResult::Invalid || ValidationErrors.Num() > 0)
 		{
 			OutValidationIcon = "MessageLog.Error";
 			OutValidationText = LOCTEXT("SourceControl.Submit.ChangelistValidationError", "Changelist validation failed!").ToString();
+			bValidationResult = false;
 		}
 		else if (ValidationResult == EDataValidationResult::NotValidated || ValidationWarnings.Num() > 0)
 		{
@@ -545,7 +550,8 @@ static void GetChangelistValidationResult(FSourceControlChangelistPtr InChangeli
 
 		int32 NumLinesDisplayed = 0;
 
-		auto AppendInfo = [&OutValidationText, &NumLinesDisplayed](const TArray<FText>& Info, const FString& InfoType) {
+		auto AppendInfo = [&OutValidationText, &NumLinesDisplayed](const TArray<FText>& Info, const FString& InfoType)
+		{
 			const int32 MaxNumLinesDisplayed = 5;
 
 			if (Info.Num() > 0)
@@ -557,6 +563,8 @@ static void GetChangelistValidationResult(FSourceControlChangelistPtr InChangeli
 				{
 					if (NumLinesDisplayed >= MaxNumLinesDisplayed)
 					{
+						OutValidationText += LINE_TERMINATOR;
+						OutValidationText += FString::Printf(TEXT("See log for complete list of %s"), *InfoType);
 						break;
 					}
 
@@ -568,9 +576,26 @@ static void GetChangelistValidationResult(FSourceControlChangelistPtr InChangeli
 			}
 		};
 
+		auto LogInfo = [](const TArray<FText>& Info, const FString& InfoType, const ELogVerbosity::Type LogVerbosity)
+		{
+			if (Info.Num() > 0)
+			{
+				FMsg::Logf(nullptr, 0, LogSourceControlChangelist.GetCategoryName(), LogVerbosity, TEXT("Encountered %d %s:"), Info.Num(), *InfoType);
+				for (const FText& Line : Info)
+				{
+					FMsg::Logf(nullptr, 0, LogSourceControlChangelist.GetCategoryName(), LogVerbosity, TEXT("%s"), *Line.ToString());
+				}
+			}
+		};
+
 		AppendInfo(ValidationErrors, TEXT("errors"));
 		AppendInfo(ValidationWarnings, TEXT("warnings"));
+
+		LogInfo(ValidationErrors, TEXT("errors"), ELogVerbosity::Error);
+		LogInfo(ValidationWarnings, TEXT("warnings"), ELogVerbosity::Warning);
 	}
+
+	return bValidationResult;
 }
 
 void SSourceControlChangelistsWidget::OnSubmitChangelist()
@@ -584,7 +609,7 @@ void SSourceControlChangelistsWidget::OnSubmitChangelist()
 
 	FString ChangelistValidationText;
 	FName ChangelistValidationIconName;
-	GetChangelistValidationResult(ChangelistState->GetChangelist(), ChangelistValidationText, ChangelistValidationIconName);
+	bool bValidationResult = GetChangelistValidationResult(ChangelistState->GetChangelist(), ChangelistValidationText, ChangelistValidationIconName);
 
 	// Build list of states for the dialog
 	FText ChangelistDescription = ChangelistState->GetDescriptionText();
@@ -606,7 +631,8 @@ void SSourceControlChangelistsWidget::OnSubmitChangelist()
 		.ChangeValidationIcon(ChangelistValidationIconName)
 		.AllowDescriptionChange(bAskForChangelistDescription)
 		.AllowUncheckFiles(false)
-		.AllowKeepCheckedOut(false);
+		.AllowKeepCheckedOut(false)
+		.AllowSubmit(bValidationResult);
 
 	NewWindow->SetContent(
 		SourceControlWidget
diff --git a/Engine/Source/Editor/SourceControlWindows/Private/SSourceControlSubmit.cpp b/Engine/Source/Editor/SourceControlWindows/Private/SSourceControlSubmit.cpp
--- a/Engine/Source/Editor/SourceControlWindows/Private/SSourceControlSubmit.cpp
+++ b/Engine/Source/Editor/SourceControlWindows/Private/SSourceControlSubmit.cpp
@@ -76,6 +76,8 @@ void SSourceControlSubmitWidget::Construct(const FArguments& InArgs)
 	SortByColumn = SSourceControlSubmitWidgetDefs::ColumnID_AssetLabel;
 	SortMode = EColumnSortMode::Ascending;
 	SavedChangeListDescription = InArgs._Description.Get();
+	bAllowSubmit = InArgs._AllowSubmit.Get();
+
 	const bool bDescriptionIsReadOnly = !InArgs._AllowDescriptionChange.Get();
 	const bool bAllowUncheckFiles = InArgs._AllowUncheckFiles.Get();
 	const bool bAllowKeepCheckedOut = InArgs._AllowKeepCheckedOut.Get();
@@ -261,9 +263,9 @@ void SSourceControlSubmitWidget::Construct(const FArguments& InArgs)
 			SNew(SButton)
 			.HAlign(HAlign_Center)
 			.ContentPadding(FEditorStyle::GetMargin("StandardDialog.ContentPadding"))
-			.IsEnabled(this, &SSourceControlSubmitWidget::IsOKEnabled)
+			.IsEnabled(this, &SSourceControlSubmitWidget::IsSubmitEnabled)
 			.Text( NSLOCTEXT("SourceControl.SubmitPanel", "OKButton", "Submit") )
-			.OnClicked(this, &SSourceControlSubmitWidget::OKClicked)
+			.OnClicked(this, &SSourceControlSubmitWidget::SubmitClicked)
 		]
 		+SUniformGridPanel::Slot(1,0)
 		[
@@ -488,7 +490,7 @@ void SSourceControlSubmitWidget::ClearChangeListDescription()
 	ChangeListDescriptionTextCtrl->SetText(FText());
 }
 
-FReply SSourceControlSubmitWidget::OKClicked()
+FReply SSourceControlSubmitWidget::SubmitClicked()
 {
 	DialogResult = ESubmitResults::SUBMIT_ACCEPTED;
 	ParentFrame.Pin()->RequestDestroyWindow();
@@ -506,15 +508,15 @@ FReply SSourceControlSubmitWidget::CancelClicked()
 }
 
 
-bool SSourceControlSubmitWidget::IsOKEnabled() const
+bool SSourceControlSubmitWidget::IsSubmitEnabled() const
 {
-	return !ChangeListDescriptionTextCtrl->GetText().IsEmpty();
+	return bAllowSubmit && !ChangeListDescriptionTextCtrl->GetText().IsEmpty();
 }
 
 
 EVisibility SSourceControlSubmitWidget::IsWarningPanelVisible() const
 {
-	return IsOKEnabled()? EVisibility::Hidden : EVisibility::Visible;
+	return IsSubmitEnabled()? EVisibility::Hidden : EVisibility::Visible;
 }
 
 
diff --git a/Engine/Source/Editor/SourceControlWindows/Private/SSourceControlSubmit.h b/Engine/Source/Editor/SourceControlWindows/Private/SSourceControlSubmit.h
--- a/Engine/Source/Editor/SourceControlWindows/Private/SSourceControlSubmit.h
+++ b/Engine/Source/Editor/SourceControlWindows/Private/SSourceControlSubmit.h
@@ -52,6 +52,7 @@ public:
 		, _AllowDescriptionChange(true)
 		, _AllowUncheckFiles(true)
 		, _AllowKeepCheckedOut(true)
+		, _AllowSubmit(true)
 	{}
 
 		SLATE_ATTRIBUTE(TSharedPtr<SWindow>, ParentWindow)
@@ -62,6 +63,7 @@ public:
 		SLATE_ATTRIBUTE(bool, AllowDescriptionChange)
 		SLATE_ATTRIBUTE(bool, AllowUncheckFiles)
 		SLATE_ATTRIBUTE(bool, AllowKeepCheckedOut)
+		SLATE_ATTRIBUTE(bool, AllowSubmit)
 
 	SLATE_END_ARGS()
 
@@ -102,13 +104,13 @@ private:
 	void OnToggleSelectedCheckBox(ECheckBoxState InNewState);
 
 	/** Called when the settings of the dialog are to be accepted*/
-	FReply OKClicked();
+	FReply SubmitClicked();
 
 	/** Called when the settings of the dialog are to be ignored*/
 	FReply CancelClicked();
 
-	/** Called to check if the OK button is enabled or not. */
-	bool IsOKEnabled() const;
+	/** Called to check if the submit button is enabled or not. */
+	bool IsSubmitEnabled() const;
 
 	/** Check if the warning panel should be visible. */
 	EVisibility IsWarningPanelVisible() const;
@@ -178,6 +180,9 @@ private:
 	/** State of the "Keep checked out" checkbox */
 	ECheckBoxState	KeepCheckedOut;
 
+	/** Whether the submit button should be enabled or not */
+	bool bAllowSubmit;
+
 	/** Specify which column to sort with */
 	FName SortByColumn;
 
