Added timer on changelist refresh operation
Added Move files operation on files in changelist dialog (incl. move to new)
Added new changelist button in toolbar
Added cancellable "path beautification" when we have deleted/shelved files with no matching file, until we make it lazy and/or batched

diff --git a/Engine/Plugins/Developer/PerforceSourceControl/Source/PerforceSourceControl/Private/PerforceSourceControlOperations.cpp b/Engine/Plugins/Developer/PerforceSourceControl/Source/PerforceSourceControl/Private/PerforceSourceControlOperations.cpp
--- a/Engine/Plugins/Developer/PerforceSourceControl/Source/PerforceSourceControl/Private/PerforceSourceControlOperations.cpp
+++ b/Engine/Plugins/Developer/PerforceSourceControl/Source/PerforceSourceControl/Private/PerforceSourceControlOperations.cpp
@@ -2107,18 +2107,33 @@ bool FPerforceNewChangelistWorker::Execute(class FPerforceSourceControlCommand&
 		check(InCommand.Operation->GetName() == GetName());
 		TSharedRef<FNewChangelist, ESPMode::ThreadSafe> Operation = StaticCastSharedRef<FNewChangelist>(InCommand.Operation);
 
-		int32 ChangeList = Connection.CreatePendingChangelist(Operation->GetDescription(), InCommand.Files, FOnIsCancelled::CreateRaw(&InCommand, &FPerforceSourceControlCommand::IsCanceled), InCommand.ResultInfo.ErrorMessages);
+		int32 ChangeList = Connection.CreatePendingChangelist(Operation->GetDescription(), TArray<FString>(), FOnIsCancelled::CreateRaw(&InCommand, &FPerforceSourceControlCommand::IsCanceled), InCommand.ResultInfo.ErrorMessages);
 
 		InCommand.bCommandSuccessful = (ChangeList > 0);
-
-		if (InCommand.bCommandSuccessful)
+		
+		// Successfully created new changelist
+		if (ChangeList > 0)
 		{
 			NewChangelist = FPerforceSourceControlChangelist(ChangeList);
 			NewChangelistState.Changelist = NewChangelist;
 			NewChangelistState.Description = Operation->GetDescription().ToString();
 			NewChangelistState.bHasShelvedFiles = false;
 
-			// Todo: keep files state also so we can update properly
+			if (InCommand.Files.Num() > 0)
+			{
+				// Move files to changelist
+				if (!RunReopenCommand(InCommand, InCommand.Files, NewChangelist, &MovedFiles))
+				{
+					// Move failed; delete newly created changelist
+					FP4RecordSet Records;
+					TArray<FString> ChangeParams;
+					ChangeParams.Add(TEXT("-d"));
+					ChangeParams.Add(NewChangelist.ToString());
+					Connection.RunCommand(TEXT("change"), ChangeParams, Records, InCommand.ResultInfo.ErrorMessages, FOnIsCancelled::CreateRaw(&InCommand, &FPerforceSourceControlCommand::IsCanceled), InCommand.bConnectionDropped);
+
+					InCommand.bCommandSuccessful = false;
+				}
+			}
 		}
 	}
 
@@ -2127,16 +2142,37 @@ bool FPerforceNewChangelistWorker::Execute(class FPerforceSourceControlCommand&
 
 bool FPerforceNewChangelistWorker::UpdateStates() const
 {
-	FPerforceSourceControlModule& PerforceSourceControl = FPerforceSourceControlModule::Get();
-	const FDateTime Now = FDateTime::Now();
+	if (NewChangelist.IsInitialized())
+	{
+		FPerforceSourceControlModule& PerforceSourceControl = FPerforceSourceControlModule::Get();
+		const FDateTime Now = FDateTime::Now();
+
+		TSharedRef<FPerforceSourceControlChangelistState, ESPMode::ThreadSafe> ChangelistState = PerforceSourceControl.GetProvider().GetStateInternal(NewChangelist);
+		*ChangelistState = NewChangelistState;
+		ChangelistState->TimeStamp = Now;
 
-	TSharedRef<FPerforceSourceControlChangelistState, ESPMode::ThreadSafe> ChangelistState = PerforceSourceControl.GetProvider().GetStateInternal(NewChangelist);
-	*ChangelistState = NewChangelistState;
-	ChangelistState->TimeStamp = Now;
+		for (const FString& MovedFile : MovedFiles)
+		{
+			TSharedRef<FPerforceSourceControlState, ESPMode::ThreadSafe> FileState = PerforceSourceControl.GetProvider().GetStateInternal(MovedFile);
 
-	// TODO: Files in new changelist support
+			// 1- Remove these files from their previous changelist
+			TSharedRef<FPerforceSourceControlChangelistState, ESPMode::ThreadSafe> PreviousChangelist = PerforceSourceControl.GetProvider().GetStateInternal(FileState->Changelist);
+			PreviousChangelist->Files.Remove(FileState);
 
-	return true;
+			// 2- Add to the new changelist
+			ChangelistState->Files.Add(FileState);
+
+			// 3- Update changelist in file state
+			FileState->Changelist = NewChangelist;
+			FileState->TimeStamp = Now;
+		}
+
+		return true;
+	}
+	else
+	{
+		return false;
+	}	
 }
 
 
diff --git a/Engine/Plugins/Developer/PerforceSourceControl/Source/PerforceSourceControl/Private/PerforceSourceControlOperations.h b/Engine/Plugins/Developer/PerforceSourceControl/Source/PerforceSourceControl/Private/PerforceSourceControlOperations.h
--- a/Engine/Plugins/Developer/PerforceSourceControl/Source/PerforceSourceControl/Private/PerforceSourceControlOperations.h
+++ b/Engine/Plugins/Developer/PerforceSourceControl/Source/PerforceSourceControl/Private/PerforceSourceControlOperations.h
@@ -220,6 +220,9 @@ public:
 	/** New changelist information */
 	FPerforceSourceControlChangelist NewChangelist;
 	FPerforceSourceControlChangelistState NewChangelistState;
+
+	/** Files that were moved */
+	TArray<FString> MovedFiles;
 };
 
 class FPerforceDeleteChangelistWorker : public IPerforceSourceControlWorker
diff --git a/Engine/Source/Editor/EditorStyle/Private/StarshipStyle.cpp b/Engine/Source/Editor/EditorStyle/Private/StarshipStyle.cpp
--- a/Engine/Source/Editor/EditorStyle/Private/StarshipStyle.cpp
+++ b/Engine/Source/Editor/EditorStyle/Private/StarshipStyle.cpp
@@ -7319,6 +7319,7 @@ void FStarshipEditorStyle::FStyle::SetupSourceControlStyles()
 		Set("SourceControl.Actions.ChangeSettings", new CORE_IMAGE_BRUSH_SVG("Starship/SourceControl/icon_SCC_Change_Source_Control_Settings", Icon16x16));
 		Set("SourceControl.Actions.CheckOut", new CORE_IMAGE_BRUSH_SVG("Starship/SourceControl/SCC_CheckedOut", Icon16x16));
 		Set("SourceControl.Actions.Refresh", new CORE_IMAGE_BRUSH_SVG("Starship/Common/Undo", Icon16x16));
+		Set("SourceControl.Actions.NewChangelist", new CORE_IMAGE_BRUSH_SVG("Starship/Common/plus-circle", Icon16x16));
 		
 	}
 #endif // WITH_EDITOR || (IS_PROGRAM && WITH_UNREAL_DEVELOPER_TOOLS)
diff --git a/Engine/Source/Editor/SourceControlWindows/Private/SSourceControlChangelists.cpp b/Engine/Source/Editor/SourceControlWindows/Private/SSourceControlChangelists.cpp
--- a/Engine/Source/Editor/SourceControlWindows/Private/SSourceControlChangelists.cpp
+++ b/Engine/Source/Editor/SourceControlWindows/Private/SSourceControlChangelists.cpp
@@ -25,6 +25,7 @@
 #include "ContentBrowserModule.h"
 #include "IContentBrowserSingleton.h"
 #include "Misc/MessageDialog.h"
+#include "Misc/ScopedSlowTask.h"
 #include "Algo/AnyOf.h"
 
 #include "SSourceControlSubmit.h"
@@ -53,6 +54,7 @@ struct FSCCFileDragDropOp : public FDragDropOperation
 //////////////////////////////
 SSourceControlChangelistsWidget::SSourceControlChangelistsWidget()
 {
+	bIsRefreshing = false;
 }
 
 void SSourceControlChangelistsWidget::Construct(const FArguments& InArgs)
@@ -92,6 +94,19 @@ void SSourceControlChangelistsWidget::Construct(const FArguments& InArgs)
 				TreeView.ToSharedRef()
 			]
 		]
+		+ SVerticalBox::Slot()
+		.AutoHeight()
+		[
+			SNew(SBorder)
+			.Visibility_Lambda([this]() -> EVisibility
+			{
+				return bIsRefreshing ? EVisibility::Visible : EVisibility::Collapsed;
+			})
+			[
+				SNew(STextBlock)
+				.Text_Lambda([this]() { return RefreshStatus; })
+			]
+		]
 	];
 
 	bShouldRefresh = true;
@@ -111,6 +126,13 @@ TSharedRef<SWidget> SSourceControlChangelistsWidget::MakeToolBar()
 		LOCTEXT("SourceControl_RefreshButton_Tooltip", "Refreshes changelists from source control provider."),
 		FSlateIcon(FEditorStyle::GetStyleSetName(), "SourceControl.Actions.Refresh"));
 
+	ToolBarBuilder.AddToolBarButton(
+		FUIAction(FExecuteAction::CreateSP(this, &SSourceControlChangelistsWidget::OnNewChangelist)),
+		NAME_None,
+		LOCTEXT("SourceControl_NewChangelistButton", "New Changelist"),
+		LOCTEXT("SourceControl_NewChangelistButton_Tooltip", "Creates an empty changelist"),
+		FSlateIcon(FEditorStyle::GetStyleSetName(), "SourceControl.Actions.Add"));
+
 	return ToolBarBuilder.MakeWidget();
 }
 
@@ -129,12 +151,19 @@ void SSourceControlChangelistsWidget::Tick(const FGeometry& AllottedGeometry, co
 			ClearChangelistsTree();
 		}
 	}
+	
+	if (bIsRefreshing)
+	{
+		TickRefreshStatus(InDeltaTime);
+	}
 }
 
 void SSourceControlChangelistsWidget::RequestRefresh()
 {
 	if (ISourceControlModule::Get().IsEnabled())
 	{
+		StartRefreshStatus();
+
 		TSharedRef<FUpdatePendingChangelistsStatus, ESPMode::ThreadSafe> UpdatePendingChangelistsOperation = ISourceControlOperation::Create<FUpdatePendingChangelistsStatus>();
 		UpdatePendingChangelistsOperation->SetUpdateAllChangelists(true);
 		UpdatePendingChangelistsOperation->SetUpdateFilesStates(true);
@@ -150,6 +179,24 @@ void SSourceControlChangelistsWidget::RequestRefresh()
 	}
 }
 
+void SSourceControlChangelistsWidget::StartRefreshStatus()
+{
+	bIsRefreshing = true;
+	RefreshStatusTimeElapsed = 0;
+}
+
+void SSourceControlChangelistsWidget::TickRefreshStatus(double InDeltaTime)
+{
+	RefreshStatusTimeElapsed += InDeltaTime;
+	const int SecondsElapsed = (int)RefreshStatusTimeElapsed;
+	RefreshStatus = FText::Format(LOCTEXT("SourceControl_RefreshStatus", "Refreshing changelists... ({0} s)"), FText::AsNumber(SecondsElapsed));
+}
+
+void SSourceControlChangelistsWidget::EndRefreshStatus()
+{
+	bIsRefreshing = false;
+}
+
 void SSourceControlChangelistsWidget::ClearChangelistsTree()
 {
 	if (!ChangelistsNodes.IsEmpty())
@@ -174,14 +221,29 @@ void SSourceControlChangelistsWidget::Refresh()
 
 		ChangelistsNodes.Reset(ChangelistsStates.Num());
 
+		// Count number of steps for slow task...
+		int32 ElementsToProcess = ChangelistsStates.Num();
+		for (FSourceControlChangelistStateRef ChangelistState : ChangelistsStates)
+		{
+			ElementsToProcess += ChangelistState->GetFilesStates().Num();
+			ElementsToProcess += ChangelistState->GetShelvedFilesStates().Num();
+		}
+
+		FScopedSlowTask SlowTask(ElementsToProcess, LOCTEXT("SourceControl_RebuildTree", "Beautifying Changelist Tree Items Paths"));
+		SlowTask.MakeDialog(/*bShowCancelButton=*/true);
+
+		bool bBeautifyPaths = true;
+
 		for (FSourceControlChangelistStateRef ChangelistState : ChangelistsStates)
 		{
 			FChangelistTreeItemRef ChangelistTreeItem = MakeShareable(new FChangelistTreeItem(ChangelistState));
 
 			for (FSourceControlStateRef FileRef : ChangelistState->GetFilesStates())
 			{
-				FChangelistTreeItemRef FileTreeItem = MakeShareable(new FFileTreeItem(FileRef));
+				FChangelistTreeItemRef FileTreeItem = MakeShareable(new FFileTreeItem(FileRef, bBeautifyPaths));
 				ChangelistTreeItem->AddChild(FileTreeItem);
+				SlowTask.EnterProgressFrame();
+				bBeautifyPaths &= !SlowTask.ShouldCancel();
 			}
 
 			if (ChangelistState->GetShelvedFilesStates().Num() > 0)
@@ -191,12 +253,16 @@ void SSourceControlChangelistsWidget::Refresh()
 
 				for (FSourceControlStateRef ShelvedFileRef : ChangelistState->GetShelvedFilesStates())
 				{
-					FChangelistTreeItemRef ShelvedFileTreeItem = MakeShareable(new FShelvedFileTreeItem(ShelvedFileRef));
+					FChangelistTreeItemRef ShelvedFileTreeItem = MakeShareable(new FShelvedFileTreeItem(ShelvedFileRef, bBeautifyPaths));
 					ShelvedChangelistTreeItem->AddChild(ShelvedFileTreeItem);
+					SlowTask.EnterProgressFrame();
+					bBeautifyPaths &= !SlowTask.ShouldCancel();
 				}
 			}
 
 			ChangelistsNodes.Add(ChangelistTreeItem);
+			SlowTask.EnterProgressFrame();
+			bBeautifyPaths &= !SlowTask.ShouldCancel();
 		}
 
 		RestoreExpandedState(ExpandedStates);
@@ -207,6 +273,8 @@ void SSourceControlChangelistsWidget::Refresh()
 	{
 		ClearChangelistsTree();
 	}
+
+	EndRefreshStatus();
 }
 
 void SSourceControlChangelistsWidget::OnSourceControlProviderChanged(ISourceControlProvider& OldProvider, ISourceControlProvider& NewProvider)
@@ -664,6 +732,76 @@ bool SSourceControlChangelistsWidget::CanSubmitChangelist()
 	return Changelist != nullptr && Changelist->GetFilesStates().Num() > 0 && Changelist->GetShelvedFilesStates().Num() == 0;
 }
 
+void SSourceControlChangelistsWidget::OnMoveFiles()
+{
+	TArray<FString> SelectedFiles = GetSelectedFiles();
+
+	if (SelectedFiles.Num() == 0)
+	{
+		return;
+	}
+
+	const bool bAddNewChangelistEntry = true;
+
+	// Build selection list for changelists
+	TArray<SSourceControlDescriptionItem> Items;
+	Items.Reset(ChangelistsNodes.Num() + (bAddNewChangelistEntry ? 1 : 0));
+
+	if (bAddNewChangelistEntry)
+	{
+		// First is always new changelist
+		Items.Emplace(
+			LOCTEXT("SourceControl_NewChangelistText", "New Changelist"),
+			LOCTEXT("SourceControl_NewChangelistDescription", "<enter description here>"),
+			/*bCanEditDescription=*/true);
+	}
+
+	const bool bCanEditAlreadyExistingChangelistDescription = false;
+
+	for (FChangelistTreeItemPtr Changelist : ChangelistsNodes)
+	{
+		if (!Changelist || Changelist->GetTreeItemType() != IChangelistTreeItem::Changelist)
+		{
+			continue;
+		}
+
+		auto TypedChangelist = StaticCastSharedPtr<FChangelistTreeItem>(Changelist);
+		Items.Emplace(TypedChangelist->GetDisplayText(), TypedChangelist->GetDescriptionText(), bCanEditAlreadyExistingChangelistDescription);
+	}
+
+	int32 PickedItem = 0;
+	FText ChangelistDescription;
+	
+	bool bOk = PickChangelistOrNewWithDescription(
+		nullptr,
+		LOCTEXT("SourceControl.MoveFiles.Title", "Move Files To..."),
+		LOCTEXT("SourceControl.MoveFIles.Label", "Target Changelist:"),
+		Items,
+		PickedItem,
+		ChangelistDescription);
+
+	if (!bOk)
+	{
+		return;
+	}
+
+	ISourceControlProvider& SourceControlProvider = ISourceControlModule::Get().GetProvider();
+
+	// Create new changelist
+	if (bAddNewChangelistEntry && PickedItem == 0)
+	{
+		auto NewChangelistOperation = ISourceControlOperation::Create<FNewChangelist>();
+		NewChangelistOperation->SetDescription(ChangelistDescription);
+		SourceControlProvider.Execute(NewChangelistOperation, SelectedFiles);
+	}
+	else
+	{
+		const int32 ChangelistIndex = (bAddNewChangelistEntry ? PickedItem - 1 : PickedItem);
+		FSourceControlChangelistPtr Changelist = StaticCastSharedPtr<FChangelistTreeItem>(ChangelistsNodes[ChangelistIndex])->ChangelistState->GetChangelist();
+		SourceControlProvider.Execute(ISourceControlOperation::Create<FMoveToChangelist>(), Changelist, SelectedFiles);
+	}
+}
+
 void SSourceControlChangelistsWidget::OnLocateFile()
 { 
 	TArray<FAssetData> AssetsToSync;
@@ -827,6 +965,10 @@ TSharedPtr<SWidget> SSourceControlChangelistsWidget::OnOpenContextMenu()
 	{
 		Section.AddSeparator("FilesSeparator");
 
+		Section.AddMenuEntry("MoveFiles", LOCTEXT("SourceControl_MoveFiles", "Move Files To..."), LOCTEXT("SourceControl_MoveFiles_Tooltip", "Move Files To A Different Changelist..."), FSlateIcon(),
+			FUIAction(
+				FExecuteAction::CreateSP(this, &SSourceControlChangelistsWidget::OnMoveFiles)));
+
 		Section.AddMenuEntry("LocateFile", LOCTEXT("SourceControl_LocateFile", "Locate File..."), LOCTEXT("SourceControl_LocateFile_Tooltip", "Locate File in Project..."), FSlateIcon(),
 			FUIAction(
 				FExecuteAction::CreateSP(this, &SSourceControlChangelistsWidget::OnLocateFile),
diff --git a/Engine/Source/Editor/SourceControlWindows/Private/SSourceControlChangelists.h b/Engine/Source/Editor/SourceControlWindows/Private/SSourceControlChangelists.h
--- a/Engine/Source/Editor/SourceControlWindows/Private/SSourceControlChangelists.h
+++ b/Engine/Source/Editor/SourceControlWindows/Private/SSourceControlChangelists.h
@@ -76,6 +76,7 @@ private:
 	void OnDeleteShelvedFiles();
 
 	/** Files operations */
+	void OnMoveFiles();
 	void OnLocateFile();
 	bool CanLocateFile();
 	void OnShowHistory();
@@ -115,4 +116,12 @@ private:
 	FDelegateHandle SourceControlStateChangedDelegateHandle;
 
 	bool bShouldRefresh;
+
+	void StartRefreshStatus();
+	void TickRefreshStatus(double InDeltaTime);
+	void EndRefreshStatus();
+
+	FText RefreshStatus;
+	bool bIsRefreshing;
+	double RefreshStatusTimeElapsed;
 };
diff --git a/Engine/Source/Editor/SourceControlWindows/Private/SSourceControlCommon.cpp b/Engine/Source/Editor/SourceControlWindows/Private/SSourceControlCommon.cpp
--- a/Engine/Source/Editor/SourceControlWindows/Private/SSourceControlCommon.cpp
+++ b/Engine/Source/Editor/SourceControlWindows/Private/SSourceControlCommon.cpp
@@ -46,7 +46,7 @@ FText FShelvedChangelistTreeItem::GetDisplayText() const
 }
 
 //////////////////////////////////////////////////////////////////////////
-FFileTreeItem::FFileTreeItem(FSourceControlStateRef InFileState, bool bIsShelvedFile)
+FFileTreeItem::FFileTreeItem(FSourceControlStateRef InFileState, bool bBeautifyPaths, bool bIsShelvedFile)
 	: FileState(InFileState)
 {
 	Type = (bIsShelvedFile ? IChangelistTreeItem::ShelvedFile : IChangelistTreeItem::File);
@@ -62,7 +62,7 @@ FFileTreeItem::FFileTreeItem(FSourceControlStateRef InFileState, bool bIsShelved
 
 	// For deleted items, the file is not on disk anymore so the only way we can get the asset data is by getting the file from the depot.
 	// For shelved files, if the file still exists locally, it will have been found before, otherwise, the history of the shelved file state will point to the remote version
-	if (FileState->IsDeleted() || (bIsShelvedFile && Assets.Num() == 0))
+	if (bBeautifyPaths && (FileState->IsDeleted() || (bIsShelvedFile && Assets.Num() == 0)))
 	{
 		// At the moment, getting the asset data from non-external assets yields issues with the package path
 		// so we will fall down to our recovery (below) instead
diff --git a/Engine/Source/Editor/SourceControlWindows/Private/SSourceControlCommon.h b/Engine/Source/Editor/SourceControlWindows/Private/SSourceControlCommon.h
--- a/Engine/Source/Editor/SourceControlWindows/Private/SSourceControlCommon.h
+++ b/Engine/Source/Editor/SourceControlWindows/Private/SSourceControlCommon.h
@@ -84,7 +84,7 @@ struct FShelvedChangelistTreeItem : public IChangelistTreeItem
 
 struct FFileTreeItem : public IChangelistTreeItem
 {
-	explicit FFileTreeItem(FSourceControlStateRef InFileState, bool bIsShelvedFile = false);
+	explicit FFileTreeItem(FSourceControlStateRef InFileState, bool bBeautifyPaths = true, bool bIsShelvedFile = false);
 
 	/** Returns the asset name of the item */
 	FText GetAssetName() const { return AssetName; }
@@ -164,8 +164,8 @@ private:
 
 struct FShelvedFileTreeItem : public FFileTreeItem
 {
-	explicit FShelvedFileTreeItem(FSourceControlStateRef InFileState)
-		: FFileTreeItem(InFileState, /*bIsShelved=*/true)
+	explicit FShelvedFileTreeItem(FSourceControlStateRef InFileState, bool bBeautifyPaths = true)
+		: FFileTreeItem(InFileState, bBeautifyPaths,/*bIsShelved=*/true)
 	{}
 };
 
diff --git a/Engine/Source/Editor/SourceControlWindows/Private/SSourceControlDescription.cpp b/Engine/Source/Editor/SourceControlWindows/Private/SSourceControlDescription.cpp
--- a/Engine/Source/Editor/SourceControlWindows/Private/SSourceControlDescription.cpp
+++ b/Engine/Source/Editor/SourceControlWindows/Private/SSourceControlDescription.cpp
@@ -5,15 +5,24 @@
 #include "Widgets/Layout/SUniformGridPanel.h"
 #include "Widgets/Text/STextBlock.h"
 #include "Widgets/Input/SButton.h"
+#include "Widgets/Input/SComboButton.h"
 #include "EditorStyleSet.h"
 #include "Framework/Application/SlateApplication.h"
-
+#include "Framework/MultiBox/MultiBoxBuilder.h"
 
 #define LOCTEXT_NAMESPACE "SourceControl.Description"
 
 void SSourceControlDescriptionWidget::Construct(const FArguments& InArgs)
 {
 	ParentWindow = InArgs._ParentWindow.Get();
+	Items = InArgs._Items;
+
+	const bool bHasItems = (Items && Items->Num() > 0);
+	const bool bShowSelectionDropDown = bHasItems;
+	const bool bDescriptionCanBeEdited = (!bHasItems || (*Items)[0].bCanEditDescription);
+	const bool bSelectTextWhenFocused = !bHasItems;
+
+	CurrentlySelectedItemIndex = 0;
 
 	ChildSlot
 	[
@@ -25,14 +34,36 @@ void SSourceControlDescriptionWidget::Construct(const FArguments& InArgs)
 			.AutoHeight()
 			.Padding(16)
 			[
-				SNew(STextBlock)
-				.Text(InArgs._Label)
+				SNew(SHorizontalBox)
+				+ SHorizontalBox::Slot()
+				.HAlign(HAlign_Left)
+				.VAlign(VAlign_Center)
+				.AutoWidth()
+				[
+					SNew(STextBlock)
+					.Text(InArgs._Label)
+				]
+				+ SHorizontalBox::Slot()
+				.HAlign(HAlign_Left)
+				.VAlign(VAlign_Center)
+				.Padding(8, 0, 0, 0)
+				[
+					SNew(SComboButton)
+					.Visibility(bShowSelectionDropDown ? EVisibility::Visible : EVisibility::Hidden)
+					.OnGetMenuContent(this, &SSourceControlDescriptionWidget::GetSelectionContent)
+					.ButtonContent()
+					[
+						SNew(STextBlock)
+						.Text(this, &SSourceControlDescriptionWidget::GetSelectedItemTitle)
+					]
+				]
 			]
 			+ SVerticalBox::Slot()
 			.Padding(FMargin(16, 0, 16, 16))
 			[
 				SAssignNew(TextBox, SMultiLineEditableTextBox)
-				.SelectAllTextWhenFocused(true)
+				.SelectAllTextWhenFocused(bSelectTextWhenFocused)
+				.IsReadOnly(!bDescriptionCanBeEdited)
 				.AutoWrapText(true)
 				.Text(InArgs._Text)
 			]
@@ -100,6 +131,78 @@ FText SSourceControlDescriptionWidget::GetDescription() const
 	return TextBox->GetText();
 }
 
+static FText GetTitleAndShortDescription(const FText& InTitle, const FText& InDescription)
+{
+	if (InDescription.IsEmptyOrWhitespace())
+	{
+		return InTitle;
+	}
+	else
+	{
+		FString TrimmedDescription = InDescription.ToString();
+		
+		const int32 TrimmedMaxLen = 50;
+		if (TrimmedDescription.Len() > TrimmedMaxLen)
+		{
+			TrimmedDescription = TrimmedDescription.Left(TrimmedMaxLen);
+			TrimmedDescription += TEXT("...");
+		}
+		
+		TrimmedDescription.ReplaceInline(TEXT("\r"), TEXT(""));
+		TrimmedDescription.ReplaceInline(TEXT("\n"), TEXT(" "));
+
+		FString TitleAndShortDescription = InTitle.ToString();
+		TitleAndShortDescription += TEXT(": ");
+		TitleAndShortDescription += TrimmedDescription;
+
+		return FText::FromString(TitleAndShortDescription);
+	}
+}
+
+FText SSourceControlDescriptionWidget::GetSelectedItemTitle() const
+{
+	if (CurrentlySelectedItemIndex >= 0 && Items && CurrentlySelectedItemIndex < Items->Num())
+	{
+		const SSourceControlDescriptionItem& Item = (*Items)[CurrentlySelectedItemIndex];
+		return GetTitleAndShortDescription(Item.Title, Item.Description);
+	}
+	else
+	{
+		return LOCTEXT("SourceControlDescription_InvalidItemTitle", "Invalid");
+	}
+}
+
+TSharedRef<SWidget> SSourceControlDescriptionWidget::GetSelectionContent()
+{
+	if (!Items || Items->Num() == 0)
+	{
+		return SNullWidget::NullWidget;
+	}
+
+	FMenuBuilder MenuBuilder(true, nullptr);
+
+	for(int32 ItemIndex = 0; ItemIndex < Items->Num(); ++ItemIndex)
+	{
+		const SSourceControlDescriptionItem& Item = (*Items)[ItemIndex];
+
+		MenuBuilder.AddMenuEntry(
+			GetTitleAndShortDescription(Item.Title, Item.Description),
+			FText(),
+			FSlateIcon(),
+			FUIAction(
+				FExecuteAction::CreateLambda([this, ItemIndex, &Item]() {
+					TextBox->SetIsReadOnly(!Item.bCanEditDescription);
+					TextBox->SetText(Item.Description);
+					
+					CurrentlySelectedItemIndex = ItemIndex;
+					})),
+			NAME_None
+		);
+	}
+
+	return MenuBuilder.MakeWidget();
+}
+
 bool GetChangelistDescription(
 	const TSharedPtr<SWidget>& ParentWidget,
 	const FText& InWindowTitle, 
@@ -137,4 +240,44 @@ bool GetChangelistDescription(
 	return SourceControlWidget->GetResult();
 }
 
+bool PickChangelistOrNewWithDescription(
+	const TSharedPtr<SWidget>& ParentWidget,
+	const FText& InWindowTitle,
+	const FText& InLabel,
+	const TArray<SSourceControlDescriptionItem>& Items,
+	int32& OutPickedIndex,
+	FText& OutDescription)
+{
+	if (Items.Num() == 0)
+	{
+		return false;
+	}
+
+	TSharedRef<SWindow> NewWindow = SNew(SWindow)
+		.Title(InWindowTitle)
+		.SizingRule(ESizingRule::UserSized)
+		.ClientSize(FVector2D(600, 400))
+		.SupportsMaximize(true)
+		.SupportsMinimize(false);
+
+	TSharedRef<SSourceControlDescriptionWidget> SourceControlWidget =
+		SNew(SSourceControlDescriptionWidget)
+		.ParentWindow(NewWindow)
+		.Label(InLabel)
+		.Text(Items[0].Description)
+		.Items(&Items);
+
+	NewWindow->SetContent(SourceControlWidget);
+
+	FSlateApplication::Get().AddModalWindow(NewWindow, ParentWidget);
+
+	if (SourceControlWidget->GetResult())
+	{
+		OutPickedIndex = SourceControlWidget->GetSelectedItemIndex();
+		OutDescription = SourceControlWidget->GetDescription();
+	}
+
+	return SourceControlWidget->GetResult();
+}
+
 #undef LOCTEXT_NAMESPACE
\ No newline at end of file
diff --git a/Engine/Source/Editor/SourceControlWindows/Private/SSourceControlDescription.h b/Engine/Source/Editor/SourceControlWindows/Private/SSourceControlDescription.h
--- a/Engine/Source/Editor/SourceControlWindows/Private/SSourceControlDescription.h
+++ b/Engine/Source/Editor/SourceControlWindows/Private/SSourceControlDescription.h
@@ -9,6 +9,17 @@
 #include "Widgets/SCompoundWidget.h"
 #include "Widgets/Input/SMultiLineEditableTextBox.h"
 
+struct SSourceControlDescriptionItem
+{
+	SSourceControlDescriptionItem(const FText& InTitle, const FText& InDescription, bool bInCanEditDescription)
+		: Title(InTitle), Description(InDescription), bCanEditDescription(bInCanEditDescription)
+	{}
+
+	FText Title;
+	FText Description;
+	bool bCanEditDescription;
+};
+
 class SSourceControlDescriptionWidget : public SCompoundWidget
 {
 public:
@@ -16,11 +27,13 @@ public:
 		: _ParentWindow()
 		, _Label()
 		, _Text()
+		, _Items()
 	{}
 
 		SLATE_ATTRIBUTE(TSharedPtr<SWindow>, ParentWindow)
 		SLATE_ATTRIBUTE(FText, Label)
 		SLATE_ATTRIBUTE(FText, Text)
+		SLATE_ARGUMENT(const TArray<SSourceControlDescriptionItem>*, Items)
 
 	SLATE_END_ARGS()
 
@@ -36,6 +49,9 @@ public:
 	/** Returns the text currently in the edit box */
 	FText GetDescription() const;
 
+	/** Returns the currently selected item index if any */
+	int32 GetSelectedItemIndex() const { return CurrentlySelectedItemIndex; }
+
 private:
 	/** Called when the settings of the dialog are to be accepted*/
 	FReply OKClicked();
@@ -43,6 +59,12 @@ private:
 	/** Called when the settings of the dialog are to be ignored*/
 	FReply CancelClicked();
 
+	/** Called to populate the dropdown */
+	TSharedRef<SWidget> GetSelectionContent();
+
+	/** Returns title of currently selected item */
+	FText GetSelectedItemTitle() const;
+
 private:
 	bool bResult = false;
 
@@ -50,10 +72,21 @@ private:
 	TWeakPtr<SWindow> ParentWindow;
 
 	TSharedPtr< SMultiLineEditableTextBox> TextBox;
+
+	const TArray<SSourceControlDescriptionItem>* Items;
+	int32 CurrentlySelectedItemIndex;
 };
 
 bool GetChangelistDescription(
 	const TSharedPtr<SWidget>& ParentWidget,
 	const FText& InWindowTitle,
 	const FText& InLabel,
+	FText& OutDescription);
+
+bool PickChangelistOrNewWithDescription(
+	const TSharedPtr<SWidget>& ParentWidget,
+	const FText& InWindowTitle,
+	const FText& InLabel,
+	const TArray<SSourceControlDescriptionItem>& Items,
+	int32& OutPickedIndex,
 	FText& OutDescription);
\ No newline at end of file
