Adds notifications on success/failure when submitting with SourceControlChangelistsWidget

diff --git a/Engine/Source/Editor/SourceControlWindows/Private/SSourceControlChangelists.cpp b/Engine/Source/Editor/SourceControlWindows/Private/SSourceControlChangelists.cpp
--- a/Engine/Source/Editor/SourceControlWindows/Private/SSourceControlChangelists.cpp
+++ b/Engine/Source/Editor/SourceControlWindows/Private/SSourceControlChangelists.cpp
@@ -6,12 +6,15 @@
 
 #include "Algo/Transform.h"
 
+#include "Logging/MessageLog.h"
+
 #include "Widgets/DeclarativeSyntaxSupport.h"
 #include "Widgets/Images/SImage.h"
 #include "Widgets/Layout/SScrollBorder.h"
 #include "Widgets/Input/SButton.h"
 #include "Widgets/Layout/SBorder.h"
 #include "Widgets/Layout/SBox.h"
+#include "Widgets/Notifications/SNotificationList.h"
 
 #include "ISourceControlProvider.h"
 #include "ISourceControlModule.h"
@@ -30,6 +33,7 @@
 
 #include "SSourceControlSubmit.h"
 #include "Framework/Application/SlateApplication.h"
+#include "Framework/Notifications/NotificationManager.h"
 
 
 #define LOCTEXT_NAMESPACE "SourceControlChangelist"
@@ -712,6 +716,7 @@ void SSourceControlChangelistsWidget::OnSubmitChangelist()
 	{
 		ISourceControlProvider& SourceControlProvider = ISourceControlModule::Get().GetProvider();
 		auto SubmitChangelistOperation = ISourceControlOperation::Create<FCheckIn>();
+		bool bCheckinSuccess = false;
 
 		// Replace description only if there was none
 		if (bAskForChangelistDescription)
@@ -722,7 +727,23 @@ void SSourceControlChangelistsWidget::OnSubmitChangelist()
 			SubmitChangelistOperation->SetDescription(Description.Description);
 		}
 
-		SourceControlProvider.Execute(SubmitChangelistOperation, ChangelistState->GetChangelist());
+		bCheckinSuccess = SourceControlProvider.Execute(SubmitChangelistOperation, ChangelistState->GetChangelist()) == ECommandResult::Succeeded;
+
+		// Setup the notification for operation feedback
+		FNotificationInfo Info(SubmitChangelistOperation->GetSuccessMessage());
+
+		// Override the notification fields for failure ones
+		if (!bCheckinSuccess)
+		{
+			Info.Text = LOCTEXT("SCC_Checkin_Failed", "Failed to check in files!");
+			Info.Image = FCoreStyle::Get().GetBrush(TEXT("MessageLog.Error"));
+		}
+		
+		Info.ExpireDuration = 8.0f;
+		Info.HyperlinkText = LOCTEXT("SCC_Checkin_ShowLog", "Show Message Log");
+		Info.Hyperlink = FSimpleDelegate::CreateLambda([]() { FMessageLog("SourceControl").Open(EMessageSeverity::Info, true); });
+
+		FSlateNotificationManager::Get().AddNotification(Info);
 	}
 }
 
@@ -803,7 +824,7 @@ void SSourceControlChangelistsWidget::OnMoveFiles()
 }
 
 void SSourceControlChangelistsWidget::OnLocateFile()
-{ 
+{
 	TArray<FAssetData> AssetsToSync;
 	TArray<FChangelistTreeItemPtr> SelectedItems = TreeView->GetSelectedItems();
 
