Enables "Locate file" only if assets are accessible from ContentBrowser.

diff --git a/Engine/Source/Editor/SourceControlWindows/Private/SSourceControlChangelists.cpp b/Engine/Source/Editor/SourceControlWindows/Private/SSourceControlChangelists.cpp
--- a/Engine/Source/Editor/SourceControlWindows/Private/SSourceControlChangelists.cpp
+++ b/Engine/Source/Editor/SourceControlWindows/Private/SSourceControlChangelists.cpp
@@ -845,7 +845,15 @@ void SSourceControlChangelistsWidget::OnLocateFile()
 
 bool SSourceControlChangelistsWidget::CanLocateFile()
 {
-	return GetSelectedFiles().Num() > 0;
+	TArray<FChangelistTreeItemPtr> SelectedItems = TreeView->GetSelectedItems();
+
+	auto HasAssetData = [](const FChangelistTreeItemPtr& SelectedItem)
+	{
+		return (SelectedItem->GetTreeItemType() == IChangelistTreeItem::File) && (StaticCastSharedPtr<FFileTreeItem>(SelectedItem)->GetAssetData().Num() > 0);
+	};
+
+	// Checks if at least one selected item has asset data (ie: accessible from ContentBrowser)
+	return SelectedItems.FindByPredicate(HasAssetData) != nullptr;
 }
 
 void SSourceControlChangelistsWidget::OnShowHistory()
