Adds validation tag in changelist description.

#changelist validated

diff --git a/Engine/Source/Editor/SourceControlWindows/Private/SSourceControlChangelists.cpp b/Engine/Source/Editor/SourceControlWindows/Private/SSourceControlChangelists.cpp
--- a/Engine/Source/Editor/SourceControlWindows/Private/SSourceControlChangelists.cpp
+++ b/Engine/Source/Editor/SourceControlWindows/Private/SSourceControlChangelists.cpp
@@ -24,6 +24,7 @@
 #include "SSourceControlDescription.h"
 #include "SourceControlWindows.h"
 #include "SourceControlHelpers.h"
+#include "SourceControlSettings.h"
 #include "AssetToolsModule.h"
 #include "ContentBrowserModule.h"
 #include "IContentBrowserSingleton.h"
@@ -38,6 +39,8 @@
 
 #define LOCTEXT_NAMESPACE "SourceControlChangelist"
 
+const FText SSourceControlChangelistsWidget::ChangelistValidatedTag = LOCTEXT("ValidationTag", "#changelist validated");
+
 DEFINE_LOG_CATEGORY_STATIC(LogSourceControlChangelist, All, All);
 
 //////////////////////////////
@@ -140,6 +143,23 @@ TSharedRef<SWidget> SSourceControlChangelistsWidget::MakeToolBar()
 	return ToolBarBuilder.MakeWidget();
 }
 
+bool SSourceControlChangelistsWidget::HasValidationTag(const FText& InChangelistDescription) const
+{
+	FString DescriptionString = InChangelistDescription.ToString();
+	FString ValidationString = ChangelistValidatedTag.ToString();
+
+	return DescriptionString.Find(ValidationString) != INDEX_NONE;
+}
+
+void SSourceControlChangelistsWidget::EditChangelistDescription(const FText& InNewChangelistDescription, const FSourceControlChangelistStatePtr& InChangelistState) const
+{
+	auto EditChangelistOperation = ISourceControlOperation::Create<FEditChangelist>();
+	ISourceControlProvider& SourceControlProvider = ISourceControlModule::Get().GetProvider();
+
+	EditChangelistOperation->SetDescription(InNewChangelistDescription);
+	SourceControlProvider.Execute(EditChangelistOperation, InChangelistState->GetChangelist());
+}
+
 void SSourceControlChangelistsWidget::Tick(const FGeometry& AllottedGeometry, const double InCurrentTime, const float InDeltaTime)
 {
 	if (bShouldRefresh)
@@ -490,11 +510,8 @@ void SSourceControlChangelistsWidget::OnEditChangelist()
 		return;
 	}
 
-	auto EditChangelistOperation = ISourceControlOperation::Create<FEditChangelist>();
-	EditChangelistOperation->SetDescription(NewChangelistDescription);
+	EditChangelistDescription(NewChangelistDescription, ChangelistState);
 
-	ISourceControlProvider& SourceControlProvider = ISourceControlModule::Get().GetProvider();
-	SourceControlProvider.Execute(EditChangelistOperation, ChangelistState->GetChangelist());
 }
 
 void SSourceControlChangelistsWidget::OnRevertUnchanged()
@@ -684,8 +701,9 @@ void SSourceControlChangelistsWidget::OnSubmitChangelist()
 	bool bValidationResult = GetChangelistValidationResult(ChangelistState->GetChangelist(), ChangelistValidationText, ChangelistValidationIconName);
 
 	// Build list of states for the dialog
-	FText ChangelistDescription = ChangelistState->GetDescriptionText();
-	const bool bAskForChangelistDescription = (ChangelistDescription.IsEmptyOrWhitespace());
+	const FText OriginalChangelistDescription = ChangelistState->GetDescriptionText();
+	const bool bAskForChangelistDescription = (OriginalChangelistDescription.IsEmptyOrWhitespace());
+	FText ChangelistDescriptionToSubmit = UpdateChangelistDescriptionToSubmitIfNeeded(bValidationResult, OriginalChangelistDescription, ChangelistState);
 
 	TSharedRef<SWindow> NewWindow = SNew(SWindow)
 		.Title(NSLOCTEXT("SourceControl.ConfirmSubmit", "Title", "Confirm changelist submit"))
@@ -698,7 +716,7 @@ void SSourceControlChangelistsWidget::OnSubmitChangelist()
 		SNew(SSourceControlSubmitWidget)
 		.ParentWindow(NewWindow)
 		.Items(ChangelistState->GetFilesStates())
-		.Description(ChangelistDescription)
+		.Description(ChangelistDescriptionToSubmit)
 		.ChangeValidationDescription(ChangelistValidationText)
 		.ChangeValidationIcon(ChangelistValidationIconName)
 		.AllowDescriptionChange(bAskForChangelistDescription)
@@ -745,6 +763,31 @@ void SSourceControlChangelistsWidget::OnSubmitChangelist()
 
 		FSlateNotificationManager::Get().AddNotification(Info);
 	}
+	// If user cancelled or closed the window, we restore the original changelist description
+	else if (HasValidationTag(ChangelistDescriptionToSubmit))
+	{
+		EditChangelistDescription(OriginalChangelistDescription, ChangelistState);
+	}
+}
+
+FText SSourceControlChangelistsWidget::UpdateChangelistDescriptionToSubmitIfNeeded(const bool bInValidationResult, const FText& InOriginalChangelistDescription, const FSourceControlChangelistStatePtr& InChangelistState) const
+{
+	if (bInValidationResult && USourceControlSettings::IsValidationTagEnabled() && (!HasValidationTag(InOriginalChangelistDescription)))
+	{
+		FStringOutputDevice Str;
+
+		Str.SetAutoEmitLineTerminator(true);
+		Str.Log(InOriginalChangelistDescription);
+		Str.Log(ChangelistValidatedTag);
+
+		FText ChangelistDescription = FText::FromString(Str);
+
+		EditChangelistDescription(ChangelistDescription, InChangelistState);
+
+		return ChangelistDescription;
+	}
+
+	return InOriginalChangelistDescription;
 }
 
 bool SSourceControlChangelistsWidget::CanSubmitChangelist()
diff --git a/Engine/Source/Editor/SourceControlWindows/Private/SSourceControlChangelists.h b/Engine/Source/Editor/SourceControlWindows/Private/SSourceControlChangelists.h
--- a/Engine/Source/Editor/SourceControlWindows/Private/SSourceControlChangelists.h
+++ b/Engine/Source/Editor/SourceControlWindows/Private/SSourceControlChangelists.h
@@ -105,7 +105,27 @@ private:
 
 	TSharedRef<SWidget> MakeToolBar();
 
+	/**
+	 * Modifies the changelist description if needed, appending validation tag.
+	 * 
+	 * @param bInValidationResult	The result of the validation step
+	 * @param InOriginalChangelistDescription	Description of the changelist before modification
+	 * @param InChangelistState	Changelist to be modified
+	 * 
+	 * @return The new changelist description
+	 */
+	FText UpdateChangelistDescriptionToSubmitIfNeeded(const bool bInValidationResult, const FText& InOriginalChangelistDescription, const FSourceControlChangelistStatePtr& InChangelistState) const;
+	
+	/** Returns true if the provided changelist description contains a validation tag. */
+	bool HasValidationTag(const FText& InChangelistDescription) const;
+
+	/** Executes an operation to updates the changelist description of the provided changelist with a new description. */
+	void EditChangelistDescription(const FText& InNewChangelistDescription, const FSourceControlChangelistStatePtr& InChangelistState) const;
+
 private:
+	/** Tag to append to a changelist that passed validation */
+	static const FText ChangelistValidatedTag;
+
 	/** Changelists (root nodes) */
 	TArray<FChangelistTreeItemPtr> ChangelistsNodes;
 
diff --git a/Engine/Source/Editor/SourceControlWindows/Private/SourceControlSettings.cpp b/Engine/Source/Editor/SourceControlWindows/Private/SourceControlSettings.cpp
new file mode 100644
--- /dev/null
+++ b/Engine/Source/Editor/SourceControlWindows/Private/SourceControlSettings.cpp
@@ -0,0 +1,8 @@
+// Copyright Epic Games, Inc. All Rights Reserved.
+
+#include "SourceControlSettings.h"
+
+bool USourceControlSettings::IsValidationTagEnabled()
+{
+	return GetDefault<USourceControlSettings>()->bEnableValidationTag;
+}
diff --git a/Engine/Source/Editor/SourceControlWindows/Public/SourceControlSettings.h b/Engine/Source/Editor/SourceControlWindows/Public/SourceControlSettings.h
new file mode 100644
--- /dev/null
+++ b/Engine/Source/Editor/SourceControlWindows/Public/SourceControlSettings.h
@@ -0,0 +1,25 @@
+// Copyright Epic Games, Inc. All Rights Reserved.
+
+#pragma once
+
+#include "CoreMinimal.h"
+#include "Engine/DeveloperSettings.h"
+
+#include "SourceControlSettings.generated.h"
+
+
+/** Settings for the Source Control Integration */
+UCLASS(config = Editor, defaultconfig, meta = (DisplayName = "Source Control"))
+class USourceControlSettings : public UDeveloperSettings
+{
+	GENERATED_BODY()
+
+public:
+	/** Helper to access easily Enable Validation Tag setting */
+	static bool IsValidationTagEnabled();
+
+public:
+	/** If enabled, adds a tag in changelist descriptions when they are validated */
+	UPROPERTY(config, EditAnywhere, Category = SourceControl, meta = (ToolTip = "Adds validation tag to changelist description on submit."))
+	bool bEnableValidationTag = true;
+};
diff --git a/Engine/Source/Editor/SourceControlWindows/SourceControlWindows.Build.cs b/Engine/Source/Editor/SourceControlWindows/SourceControlWindows.Build.cs
--- a/Engine/Source/Editor/SourceControlWindows/SourceControlWindows.Build.cs
+++ b/Engine/Source/Editor/SourceControlWindows/SourceControlWindows.Build.cs
@@ -6,6 +6,10 @@ public class SourceControlWindows : ModuleRules
 {
 	public SourceControlWindows(ReadOnlyTargetRules Target) : base(Target)
 	{
+		PublicIncludePaths.Add("Editor/SourceControlWindows/Public");
+
+		PrivateIncludePaths.Add("Editor/SourceControlWindows/Private");
+
 		PrivateDependencyModuleNames.AddRange(
 			new string[] {
 				"Core", 
@@ -19,7 +23,8 @@ public class SourceControlWindows : ModuleRules
 				"AssetTools",
 				"EditorFramework",
 				"WorkspaceMenuStructure",
-                "UnrealEd"			// We need this dependency here because we use PackageTools.
+				"DeveloperSettings",
+				"UnrealEd"			// We need this dependency here because we use PackageTools.
 			}
 		);
 
