Adds Offline files to Uncontrolled Changelists
* Adds Offline files to Source Control Changelist Window
* Use TSet.Remove return value to update state.

#changelist validated

diff --git a/Engine/Source/Developer/UncontrolledChangelists/Private/UncontrolledChangelistState.cpp b/Engine/Source/Developer/UncontrolledChangelists/Private/UncontrolledChangelistState.cpp
--- a/Engine/Source/Developer/UncontrolledChangelists/Private/UncontrolledChangelistState.cpp
+++ b/Engine/Source/Developer/UncontrolledChangelists/Private/UncontrolledChangelistState.cpp
@@ -2,11 +2,13 @@
 
 #include "UncontrolledChangelistState.h"
 
+#include "Algo/Copy.h"
 #include "Algo/Transform.h"
 #include "Dom/JsonObject.h"
 #include "Dom/JsonValue.h"
 #include "ISourceControlModule.h"
 #include "ISourceControlProvider.h"
+#include "SourceControlHelpers.h"
 #include "SourceControlOperations.h"
 #include "UncontrolledChangelist.h"
 
@@ -52,14 +54,17 @@ const TSet<FSourceControlStateRef>& FUncontrolledChangelistState::GetFilesStates
 	return Files;
 }
 
+const TSet<FString>& FUncontrolledChangelistState::GetOfflineFiles() const
+{
+	return OfflineFiles;
+}
+
 void FUncontrolledChangelistState::Serialize(TSharedRef<FJsonObject> OutJsonObject) const
 {
 	TArray<TSharedPtr<FJsonValue>> FileValues;
 
-	for (const FSourceControlStateRef& File : Files)
-	{
-		FileValues.Add(MakeShareable(new FJsonValueString(File->GetFilename())));
-	}
+	Algo::Transform(Files, FileValues, [](const FSourceControlStateRef& File) { return MakeShareable(new FJsonValueString(File->GetFilename())); });
+	Algo::Transform(OfflineFiles, FileValues, [](const FString& OfflineFile) { return MakeShareable(new FJsonValueString(OfflineFile)); });
 
 	OutJsonObject->SetArrayField(FILES_NAME, MoveTemp(FileValues));
 }
@@ -67,7 +72,6 @@ void FUncontrolledChangelistState::Serialize(TSharedRef<FJsonObject> OutJsonObje
 bool FUncontrolledChangelistState::Deserialize(const TSharedRef<FJsonObject> InJsonValue)
 {
 	const TArray<TSharedPtr<FJsonValue>>* FileValues = nullptr;
-	ISourceControlProvider& SourceControlProvider = ISourceControlModule::Get().GetProvider();
 
 	if ((!InJsonValue->TryGetArrayField(FILES_NAME, FileValues)) || (FileValues == nullptr))
 	{
@@ -100,6 +104,14 @@ bool FUncontrolledChangelistState::AddFiles(const TArray<FString>& InFilenames,
 		return bOutChanged;
 	}
 
+	// No source control is available, add files to the offline file set.
+	if (!SourceControlProvider.IsAvailable())
+	{
+		int32 OldSize = OfflineFiles.Num();
+		Algo::Copy(SourceControlHelpers::AbsoluteFilenames(InFilenames), OfflineFiles);
+		return OldSize != OfflineFiles.Num();
+	}
+
 	if (bCheckStatus)
 	{
 		auto UpdateStatusOperation = ISourceControlOperation::Create<FUpdateStatus>();
@@ -133,11 +145,7 @@ bool FUncontrolledChangelistState::RemoveFiles(const TArray<FSourceControlStateR
 
 	for (const FSourceControlStateRef& FileState : InFileStates)
 	{
-		int32 OldSize = Files.Num();
-
-		Files.Remove(FileState);
-
-		bOutChanged |= OldSize != Files.Num();
+		bOutChanged |= (Files.Remove(FileState) > 0);
 	}
 
 	return bOutChanged;
@@ -147,30 +155,26 @@ bool FUncontrolledChangelistState::UpdateStatus()
 {
 	TArray<FString> FilesToUpdate;
 	bool bOutChanged = false;
+	int32 InitialFileNumber = Files.Num();
+	int32 InitialOfflineFileNumber = OfflineFiles.Num();
 
 	Algo::Transform(Files, FilesToUpdate, [](const FSourceControlStateRef& State) { return State->GetFilename(); });
+	Algo::Copy(OfflineFiles, FilesToUpdate);
+
+	Files.Empty();
+	OfflineFiles.Empty();
 
 	if (FilesToUpdate.Num() == 0)
 	{
 		return bOutChanged;
 	}
 
-	ISourceControlProvider& SourceControlProvider = ISourceControlModule::Get().GetProvider();
-	auto UpdateStatusOperation = ISourceControlOperation::Create<FUpdateStatus>();
-	UpdateStatusOperation->SetUpdateModifiedState(true);
+	bOutChanged |= AddFiles(FilesToUpdate, ECheckFlags::All);
 
-	SourceControlProvider.Execute(UpdateStatusOperation, FilesToUpdate);
+	bool bFileNumberChanged = InitialFileNumber == Files.Num();
+	bool bOfflineFileNumberChanged = InitialOfflineFileNumber == OfflineFiles.Num();
 
-	for (auto It = Files.CreateIterator(); It; ++It)
-	{
-		FSourceControlStateRef& State = *It;
-
-		if (State->IsCheckedOut() || (!State->IsModified()))
-		{
-			It.RemoveCurrent();
-			bOutChanged = true;
-		}
-	}
+	bOutChanged |= bFileNumberChanged || bOfflineFileNumberChanged;
 
 	return bOutChanged;
 }
diff --git a/Engine/Source/Developer/UncontrolledChangelists/Public/UncontrolledChangelistState.h b/Engine/Source/Developer/UncontrolledChangelists/Public/UncontrolledChangelistState.h
--- a/Engine/Source/Developer/UncontrolledChangelists/Public/UncontrolledChangelistState.h
+++ b/Engine/Source/Developer/UncontrolledChangelists/Public/UncontrolledChangelistState.h
@@ -68,6 +68,8 @@ public:
 	const FDateTime& GetTimeStamp() const;
 
 	UNCONTROLLEDCHANGELISTS_API const TSet<FSourceControlStateRef>& GetFilesStates() const;
+	
+	UNCONTROLLEDCHANGELISTS_API const TSet<FString>& GetOfflineFiles() const;
 
 	/**
 	 * Serialize the state of the Uncontrolled Changelist to a Json Object.
@@ -114,6 +116,7 @@ public:
 	FUncontrolledChangelist Changelist;
 	FString Description;
 	TSet<FSourceControlStateRef> Files;
+	TSet<FString> OfflineFiles;
 
 	/** The timestamp of the last update */
 	FDateTime TimeStamp;
diff --git a/Engine/Source/Editor/EditorStyle/Private/StarshipStyle.cpp b/Engine/Source/Editor/EditorStyle/Private/StarshipStyle.cpp
--- a/Engine/Source/Editor/EditorStyle/Private/StarshipStyle.cpp
+++ b/Engine/Source/Editor/EditorStyle/Private/StarshipStyle.cpp
@@ -7258,6 +7258,7 @@ void FStarshipEditorStyle::FStyle::SetupSourceControlStyles()
 		Set("SourceControl.ShelvedChangelist", new CORE_IMAGE_BRUSH_SVG("Starship/SourceControl/SCC_Changelist", Icon16x16, FStyleColors::AccentBlue));
 		Set("SourceControl.UncontrolledChangelist", new CORE_IMAGE_BRUSH_SVG("Starship/SourceControl/SCC_Changelist", Icon32x32, FStyleColors::AccentGray));
 		Set("SourceControl.UncontrolledChangelist_Small", new CORE_IMAGE_BRUSH_SVG("Starship/SourceControl/SCC_Changelist", Icon16x16, FStyleColors::AccentGray));
+		Set("SourceControl.OfflineFile_Small", new CORE_IMAGE_BRUSH_SVG("Starship/SourceControl/icon_SCC_Connect", Icon16x16, FStyleColors::AccentRed));
 
 		Set("SourceControl.Add", new CORE_IMAGE_BRUSH_SVG( "Starship/SourceControl/SCC_ContentAdd",Icon16x16));
 		Set("SourceControl.Edit", new CORE_IMAGE_BRUSH_SVG( "Starship/SourceControl/SCC_CheckedOut", Icon16x16));
diff --git a/Engine/Source/Editor/SourceControlWindows/Private/SSourceControlChangelists.cpp b/Engine/Source/Editor/SourceControlWindows/Private/SSourceControlChangelists.cpp
--- a/Engine/Source/Editor/SourceControlWindows/Private/SSourceControlChangelists.cpp
+++ b/Engine/Source/Editor/SourceControlWindows/Private/SSourceControlChangelists.cpp
@@ -103,7 +103,9 @@ void SSourceControlChangelistsWidget::Construct(const FArguments& InArgs)
 		+ SVerticalBox::Slot()
 		[
 			SNew(SScrollBorder, TreeView.ToSharedRef())
-			.Visibility(TAttribute<EVisibility>::Create(TAttribute<EVisibility>::FGetter::CreateLambda([]()->EVisibility { return ISourceControlModule::Get().IsEnabled() ? EVisibility::Visible : EVisibility::Hidden; })))
+			.Visibility(TAttribute<EVisibility>::Create(TAttribute<EVisibility>::FGetter::CreateLambda([]()->EVisibility { return (ISourceControlModule::Get().IsEnabled() || FUncontrolledChangelistsModule::Get().IsEnabled())
+																																   ? EVisibility::Visible
+																																   : EVisibility::Hidden; })))
 			[
 				TreeView.ToSharedRef()
 			]
@@ -189,7 +191,7 @@ void SSourceControlChangelistsWidget::Tick(const FGeometry& AllottedGeometry, co
 {
 	if (bShouldRefresh)
 	{
-		if (ISourceControlModule::Get().IsEnabled())
+		if (ISourceControlModule::Get().IsEnabled() || FUncontrolledChangelistsModule::Get().IsEnabled())
 		{
 			RequestRefresh();
 			bShouldRefresh = false;
@@ -209,8 +211,11 @@ void SSourceControlChangelistsWidget::Tick(const FGeometry& AllottedGeometry, co
 
 void SSourceControlChangelistsWidget::RequestRefresh()
 {
+	bool bAnyProviderAvailable = false;
+
 	if (ISourceControlModule::Get().IsEnabled())
 	{
+		bAnyProviderAvailable = true;
 		StartRefreshStatus();
 
 		TSharedRef<FUpdatePendingChangelistsStatus, ESPMode::ThreadSafe> UpdatePendingChangelistsOperation = ISourceControlOperation::Create<FUpdatePendingChangelistsStatus>();
@@ -220,11 +225,17 @@ void SSourceControlChangelistsWidget::RequestRefresh()
 
 		ISourceControlProvider& SourceControlProvider = ISourceControlModule::Get().GetProvider();
 		SourceControlProvider.Execute(UpdatePendingChangelistsOperation, EConcurrency::Asynchronous);
+	}
+
+	if (FUncontrolledChangelistsModule::Get().IsEnabled())
+	{
+		bAnyProviderAvailable = true;
 
 		FUncontrolledChangelistsModule& UncontrolledChangelistModule = FUncontrolledChangelistsModule::Get();
 		UncontrolledChangelistModule.UpdateStatus();
 	}
-	else
+
+	if (!bAnyProviderAvailable)
 	{
 		// No provider available, clear changelist tree
 		ClearChangelistsTree();
@@ -260,7 +271,7 @@ void SSourceControlChangelistsWidget::ClearChangelistsTree()
 
 void SSourceControlChangelistsWidget::Refresh()
 {
-	if (ISourceControlModule::Get().IsEnabled())
+	if (ISourceControlModule::Get().IsEnabled() || FUncontrolledChangelistsModule::Get().IsEnabled())
 	{
 		TMap<FSourceControlChangelistStateRef, ExpandedState> ExpandedStates;
 		SaveExpandedState(ExpandedStates);
@@ -288,6 +299,7 @@ void SSourceControlChangelistsWidget::Refresh()
 		for (FUncontrolledChangelistStateRef UncontrolledChangelistState : UncontrolledChangelistStates)
 		{
 			ElementsToProcess += UncontrolledChangelistState->GetFilesStates().Num();
+			ElementsToProcess += UncontrolledChangelistState->GetOfflineFiles().Num();
 		}
 
 		FScopedSlowTask SlowTask(ElementsToProcess, LOCTEXT("SourceControl_RebuildTree", "Beautifying Changelist Tree Items Paths"));
@@ -330,7 +342,7 @@ void SSourceControlChangelistsWidget::Refresh()
 		{
 			FChangelistTreeItemRef UncontrolledChangelistTreeItem = MakeShareable(new FUncontrolledChangelistTreeItem(UncontrolledChangelistState));
 			
-			for (FSourceControlStateRef FileRef : UncontrolledChangelistState->GetFilesStates())
+			for (const FSourceControlStateRef& FileRef : UncontrolledChangelistState->GetFilesStates())
 			{
 				FChangelistTreeItemRef FileTreeItem = MakeShareable(new FFileTreeItem(FileRef, bBeautifyPaths));
 				UncontrolledChangelistTreeItem->AddChild(FileTreeItem);
@@ -338,6 +350,14 @@ void SSourceControlChangelistsWidget::Refresh()
 				bBeautifyPaths &= !SlowTask.ShouldCancel();
 			}
 
+			for (const FString& Filename : UncontrolledChangelistState->GetOfflineFiles())
+			{
+				FChangelistTreeItemRef OfflineFileTreeItem = MakeShareable(new FOfflineFileTreeItem(Filename));
+				UncontrolledChangelistTreeItem->AddChild(OfflineFileTreeItem);
+				SlowTask.EnterProgressFrame();
+				bBeautifyPaths &= !SlowTask.ShouldCancel();
+			}
+
 			ChangelistsNodes.Add(UncontrolledChangelistTreeItem);
 			SlowTask.EnterProgressFrame();
 			bBeautifyPaths &= !SlowTask.ShouldCancel();
@@ -565,7 +585,6 @@ void SSourceControlChangelistsWidget::OnEditChangelist()
 	}
 
 	EditChangelistDescription(NewChangelistDescription, ChangelistState);
-
 }
 
 void SSourceControlChangelistsWidget::OnRevertUnchanged()
@@ -1500,6 +1519,107 @@ private:
 	FFileTreeItem* TreeItem;
 };
 
+class SOfflineFileTableRow : public SMultiColumnTableRow<FChangelistTreeItemPtr>
+{
+public:
+	SLATE_BEGIN_ARGS(SOfflineFileTableRow)
+		: _TreeItemToVisualize()
+	{
+	}
+	SLATE_ARGUMENT(FChangelistTreeItemPtr, TreeItemToVisualize)
+	SLATE_END_ARGS()
+
+public:
+	/**
+	* Construct child widgets that comprise this widget.
+	*
+	* @param InArgs Declaration from which to construct this widget.
+	*/
+	void Construct(const FArguments& InArgs, const TSharedRef<STableViewBase>& InOwner)
+	{
+		TreeItem = static_cast<FOfflineFileTreeItem*>(InArgs._TreeItemToVisualize.Get());
+
+		auto Args = FSuperRowType::FArguments().ShowSelection(true);
+		FSuperRowType::Construct(Args, InOwner);
+	}
+
+	// SMultiColumnTableRow overrides
+	virtual TSharedRef<SWidget> GenerateWidgetForColumn(const FName& ColumnName) override
+	{
+		if (ColumnName == TEXT("Change")) // eq. to name
+		{
+			return SNew(SHorizontalBox)
+
+				// Icon
+				+ SHorizontalBox::Slot()
+				.AutoWidth()
+				.Padding(40, 0, 4, 0)
+				[
+					SNew(SImage)
+					.Image(FEditorStyle::GetBrush(FName("SourceControl.OfflineFile_Small")))
+				]
+
+			+ SHorizontalBox::Slot()
+				.AutoWidth()
+				.VAlign(VAlign_Center)
+				[
+					SNew(STextBlock)
+					.Text(this, &SOfflineFileTableRow::GetDisplayName)
+				];
+		}
+		else if (ColumnName == TEXT("Description")) // eq. to path
+		{
+			return SNew(STextBlock)
+				.Text(this, &SOfflineFileTableRow::GetDisplayPath)
+				.ToolTipText(this, &SOfflineFileTableRow::GetFilename);
+		}
+		else if (ColumnName == TEXT("Type"))
+		{
+			return SNew(STextBlock)
+				.Text(this, &SOfflineFileTableRow::GetDisplayType)
+				.ColorAndOpacity(this, &SOfflineFileTableRow::GetDisplayColor);
+		}
+		else
+		{
+			return SNullWidget::NullWidget;
+		}
+	}
+
+	FText GetDisplayName() const
+	{
+		return TreeItem->GetDisplayName();
+	}
+
+	FText GetFilename() const
+	{
+		return TreeItem->GetPackageName();
+	}
+
+	FText GetDisplayPath() const
+	{
+		return TreeItem->GetDisplayPath();
+	}
+
+	FText GetDisplayType() const
+	{
+		return TreeItem->GetDisplayType();
+	}
+
+	FSlateColor GetDisplayColor() const
+	{
+		return TreeItem->GetDisplayColor();
+	}
+
+protected:
+	//~ Begin STableRow Interface.
+
+	//~ End STableRow Interface.
+
+private:
+	/** The info about the widget that we are visualizing. */
+	FOfflineFileTreeItem* TreeItem;
+};
+
 class SShelvedChangelistTableRow : public SMultiColumnTableRow<FChangelistTreeItemPtr>
 {
 public:
@@ -1586,6 +1706,10 @@ TSharedRef<ITableRow> SSourceControlChangelistsWidget::OnGenerateRow(FChangelist
 			.TreeItemToVisualize(InTreeItem)
 			.OnDragDetected(this, &SSourceControlChangelistsWidget::OnFilesDragged);
 
+	case IChangelistTreeItem::OfflineFile:
+		return SNew(SOfflineFileTableRow, OwnerTable)
+			.TreeItemToVisualize(InTreeItem);
+
 	case IChangelistTreeItem::ShelvedChangelist:
 		return SNew(SShelvedChangelistTableRow, OwnerTable)
 			.TreeItemToVisualize(InTreeItem);
diff --git a/Engine/Source/Editor/SourceControlWindows/Private/SSourceControlCommon.cpp b/Engine/Source/Editor/SourceControlWindows/Private/SSourceControlCommon.cpp
--- a/Engine/Source/Editor/SourceControlWindows/Private/SSourceControlCommon.cpp
+++ b/Engine/Source/Editor/SourceControlWindows/Private/SSourceControlCommon.cpp
@@ -84,9 +84,9 @@ FFileTreeItem::FFileTreeItem(FSourceControlStateRef InFileState, bool bBeautifyP
 	}
 
 	// Initialize display-related members
-	FString TempAssetName = LOCTEXT("SourceControl_DefaultAssetName", "Unavailable").ToString();
+	FString TempAssetName = SSourceControlCommon::GetDefaultAssetName().ToString();
 	FString TempAssetPath = Filename;
-	FString TempAssetType = LOCTEXT("SourceControl_DefaultAssetType", "Unknown").ToString();
+	FString TempAssetType = SSourceControlCommon::GetDefaultAssetType().ToString();
 	FString TempPackageName = Filename;
 	FColor TempAssetColor = FColor(		// Copied from ContentBrowserCLR.cpp
 		127 + FColor::Red.R / 2,	// Desaturate the colors a bit (GB colors were too.. much)
@@ -133,7 +133,7 @@ FFileTreeItem::FFileTreeItem(FSourceControlStateRef InFileState, bool bBeautifyP
 		}
 		else
 		{
-			TempAssetType = LOCTEXT("SourceCOntrol_ManyAssetType", "Multiple Assets").ToString();
+			TempAssetType = SSourceControlCommon::GetDefaultMultipleAsset().ToString();
 			TempAssetColor = FColor::White;
 		}
 
@@ -156,7 +156,7 @@ FFileTreeItem::FFileTreeItem(FSourceControlStateRef InFileState, bool bBeautifyP
 	{
 		TempAssetName = FPaths::GetCleanFilename(Filename);
 		TempPackageName = Filename; // put back original package name if the try failed
-		TempAssetType = FText::Format(LOCTEXT("SourceControl_FileTypeDefault", "{0} File"), FText::FromString(FPaths::GetExtension(Filename).ToUpper())).ToString();
+		TempAssetType = FText::Format(SSourceControlCommon::GetDefaultUnknownAssetType(), FText::FromString(FPaths::GetExtension(Filename).ToUpper())).ToString();
 	}
 
 	// Finally, assign the temp variables to the member variables
@@ -167,6 +167,74 @@ FFileTreeItem::FFileTreeItem(FSourceControlStateRef InFileState, bool bBeautifyP
 	PackageName = FText::FromString(TempPackageName);
 }
 
+//////////////////////////////////////////////////////////////////////////
+
+FOfflineFileTreeItem::FOfflineFileTreeItem(const FString& InFilename)
+	: Assets()
+	, PackageName(FText::FromString(InFilename)) 
+	, AssetName(SSourceControlCommon::GetDefaultAssetName())
+	, AssetPath()
+	, AssetType(SSourceControlCommon::GetDefaultAssetType())
+	, AssetTypeColor()
+{
+	Type = IChangelistTreeItem::OfflineFile;
+	FString TempString;
+
+	USourceControlHelpers::GetAssetData(InFilename, Assets);
+
+	if (Assets.Num() > 0)
+	{
+		const FAssetData& AssetData = Assets[0];
+		AssetPath = FText::FromName(AssetData.ObjectPath);
+
+		// Find name, asset type & color only if there is exactly one asset
+		if (Assets.Num() == 1)
+		{
+			static FName NAME_ActorLabel(TEXT("ActorLabel"));
+			if (AssetData.FindTag(NAME_ActorLabel))
+			{
+				AssetData.GetTagValue(NAME_ActorLabel, AssetName);
+			}
+			else
+			{
+				AssetName = FText::FromName(AssetData.AssetName);
+			}
+
+			AssetType = FText::FromName(AssetData.AssetClass);
+
+			const FAssetToolsModule& AssetToolsModule = FModuleManager::LoadModuleChecked<FAssetToolsModule>(TEXT("AssetTools"));
+			const TSharedPtr<IAssetTypeActions> AssetTypeActions = AssetToolsModule.Get().GetAssetTypeActionsForClass(AssetData.GetClass()).Pin();
+			if (AssetTypeActions.IsValid())
+			{
+				AssetTypeColor = AssetTypeActions->GetTypeColor();
+			}
+			else
+			{
+				AssetTypeColor = FColor::White;
+			}
+		}
+		else
+		{
+			AssetType = SSourceControlCommon::GetDefaultMultipleAsset();
+			AssetTypeColor = FColor::White;
+		}
+
+		// Beautify the package name
+		PackageName = AssetPath;
+	}
+	else if (FPackageName::TryConvertFilenameToLongPackageName(InFilename, TempString))
+	{
+		PackageName = FText::FromString(TempString);
+		// Fake asset name, asset path from the package name
+		AssetPath = PackageName;
+	}
+	else
+	{
+		AssetName = FText::FromString(FPaths::GetCleanFilename(InFilename));
+		AssetType = FText::Format(SSourceControlCommon::GetDefaultUnknownAssetType(), FText::FromString(FPaths::GetExtension(InFilename).ToUpper()));
+	}
+}
+
 //////////////////////////////////////////////////////////////////////////
 namespace SSourceControlCommon
 {
@@ -203,6 +271,26 @@ TSharedRef<SWidget> GetSCCFileWidget(FSourceControlStateRef InFileState, bool bI
 		];
 }
 
+FText GetDefaultAssetName()
+{
+	return LOCTEXT("SourceControl_DefaultAssetName", "Unavailable");
+}
+
+FText GetDefaultAssetType()
+{
+	return LOCTEXT("SourceControl_DefaultAssetType", "Unknown");
+}
+
+FText GetDefaultUnknownAssetType()
+{
+	return LOCTEXT("SourceControl_FileTypeDefault", "{0} File");
+}
+
+FText GetDefaultMultipleAsset()
+{
+	return LOCTEXT("SourceCOntrol_ManyAssetType", "Multiple Assets");
+}
+
 } // end of namespace SSourceControlCommon
 
 #undef LOCTEXT_NAMESPACE
\ No newline at end of file
diff --git a/Engine/Source/Editor/SourceControlWindows/Private/SSourceControlCommon.h b/Engine/Source/Editor/SourceControlWindows/Private/SSourceControlCommon.h
--- a/Engine/Source/Editor/SourceControlWindows/Private/SSourceControlCommon.h
+++ b/Engine/Source/Editor/SourceControlWindows/Private/SSourceControlCommon.h
@@ -21,6 +21,7 @@ struct IChangelistTreeItem : TSharedFromThis<IChangelistTreeItem>
 		Changelist,
 		UncontrolledChangelist,
 		File,
+		OfflineFile,
 		ShelvedChangelist, // container for shelved files
 		ShelvedFile
 	};
@@ -188,6 +189,30 @@ private:
 	TArray<FAssetData> Assets;
 };
 
+struct FOfflineFileTreeItem : public IChangelistTreeItem
+{
+	explicit FOfflineFileTreeItem(const FString& InFilename);
+
+public:
+	const FText& GetPackageName() const { return PackageName; }
+
+	const FText& GetDisplayName() const { return AssetName; }
+
+	const FText& GetDisplayPath() const { return AssetPath; }
+
+	const FText& GetDisplayType() const { return AssetType; }
+
+	const FSlateColor& GetDisplayColor() const { return AssetTypeColor; }
+
+private:
+	TArray<FAssetData> Assets;
+	FText PackageName;
+	FText AssetName;
+	FText AssetPath;
+	FText AssetType;
+	FSlateColor AssetTypeColor;
+};
+
 typedef TSharedPtr<FFileTreeItem> FFileTreeItemPtr;
 typedef TSharedRef<FFileTreeItem> FFileTreeItemRef;
 
@@ -201,4 +226,8 @@ struct FShelvedFileTreeItem : public FFileTreeItem
 namespace SSourceControlCommon
 {
 	TSharedRef<SWidget> GetSCCFileWidget(FSourceControlStateRef InFileState, bool bIsShelvedFile = false);
+	FText GetDefaultAssetName();
+	FText GetDefaultAssetType();
+	FText GetDefaultUnknownAssetType();
+	FText GetDefaultMultipleAsset();
 }
\ No newline at end of file
diff --git a/Engine/Source/Editor/SourceControlWindows/Private/SourceControlWindowsModule.cpp b/Engine/Source/Editor/SourceControlWindows/Private/SourceControlWindowsModule.cpp
--- a/Engine/Source/Editor/SourceControlWindows/Private/SourceControlWindowsModule.cpp
+++ b/Engine/Source/Editor/SourceControlWindows/Private/SourceControlWindowsModule.cpp
@@ -15,6 +15,7 @@
 #include "WorkspaceMenuStructureModule.h"
 
 #include "SSourceControlChangelists.h"
+#include "UncontrolledChangelistsModule.h"
 
 #define LOCTEXT_NAMESPACE "SourceControlWindows"
 
@@ -117,13 +118,8 @@ void FSourceControlWindowsModule::ShowChangelistsTab()
 bool FSourceControlWindowsModule::CanShowChangelistsTab() const
 {
 	ISourceControlModule& SourceControlModule = ISourceControlModule::Get();
-	if (ISourceControlModule::Get().IsEnabled() &&
-		ISourceControlModule::Get().GetProvider().IsAvailable())
-	{
-		return true;
-	}
 
-	return false;
+	return (SourceControlModule.IsEnabled() && SourceControlModule.GetProvider().IsAvailable()) || FUncontrolledChangelistsModule::Get().IsEnabled();
 }
 
 #undef LOCTEXT_NAMESPACE
