Update contextual menu in Changelist window to support Uncontrolled Changelists
* Update MarkForAdd operation to operate on a specified changelist
* Adds Force option in sync operation to be able to revert uncontrolled files
* Adds the resulting new changelist to FNewChangelist operation
* Adds support for "mark for add" files when moving from UCL to C
* Adds support for reverting Uncontrolled Changelist, Uncontrolled Changelist files, or mixes of controlled/uncontrolled files
* Adds support of "Move files to..." contextual action for Uncontrolled files and to Uncontrolled Changelists

#changelist validated

diff --git a/Engine/Plugins/Developer/PerforceSourceControl/Source/PerforceSourceControl/Private/PerforceSourceControlOperations.cpp b/Engine/Plugins/Developer/PerforceSourceControl/Source/PerforceSourceControl/Private/PerforceSourceControlOperations.cpp
--- a/Engine/Plugins/Developer/PerforceSourceControl/Source/PerforceSourceControl/Private/PerforceSourceControlOperations.cpp
+++ b/Engine/Plugins/Developer/PerforceSourceControl/Source/PerforceSourceControl/Private/PerforceSourceControlOperations.cpp
@@ -185,6 +185,26 @@ static void ParseRecordSetForState(const FP4RecordSet& InRecords, TMap<FString,
 	}
 }
 
+static void UpdateChangelistState(const FPerforceSourceControlChangelist& InChangelist, const TMap<FString, EPerforceState::Type>& InOperationResults)
+{
+	if (InChangelist.IsInitialized())
+	{
+		FPerforceSourceControlModule& PerforceSourceControl = FPerforceSourceControlModule::Get();
+		TSharedRef<FPerforceSourceControlChangelistState, ESPMode::ThreadSafe> ChangelistState = PerforceSourceControl.GetProvider().GetStateInternal(InChangelist);
+
+		for (TMap<FString, EPerforceState::Type>::TConstIterator It(InOperationResults); It; ++It)
+		{
+			if ((It.Value() != EPerforceState::CheckedOut) && (It.Value() != EPerforceState::OpenForAdd))
+			{
+				continue;
+			}
+
+			TSharedRef<FPerforceSourceControlState, ESPMode::ThreadSafe> State = PerforceSourceControl.GetProvider().GetStateInternal(It.Key());
+			ChangelistState->Files.Add(State);
+		}
+	}
+}
+
 static bool UpdateCachedStates(const TMap<FString, EPerforceState::Type>& InResults)
 {
 	FPerforceSourceControlModule& PerforceSourceControl = FPerforceSourceControlModule::Get();
@@ -339,22 +359,7 @@ bool FPerforceCheckOutWorker::Execute(FPerforceSourceControlCommand& InCommand)
 bool FPerforceCheckOutWorker::UpdateStates() const
 {
 	// If files have been checkedout directly to a CL, modify the cached state to reflect it.
-	if (InChangelist.IsInitialized())
-	{
-		FPerforceSourceControlModule& PerforceSourceControl = FPerforceSourceControlModule::Get();
-		TSharedRef<FPerforceSourceControlChangelistState, ESPMode::ThreadSafe> ChangelistState = PerforceSourceControl.GetProvider().GetStateInternal(InChangelist);
-
-		for (TMap<FString, EPerforceState::Type>::TConstIterator It(OutResults); It; ++It)
-		{
-			if (It.Value() != EPerforceState::CheckedOut)
-			{
-				continue;
-			}
-
-			TSharedRef<FPerforceSourceControlState, ESPMode::ThreadSafe> State = PerforceSourceControl.GetProvider().GetStateInternal(It.Key());
-			ChangelistState->Files.Add(State);
-		}
-	}
+	UpdateChangelistState(InChangelist, OutResults);
 
 	return UpdateCachedStates(OutResults);
 }
@@ -632,6 +637,18 @@ bool FPerforceMarkForAddWorker::Execute(FPerforceSourceControlCommand& InCommand
 		TArray<FString> Parameters;
 		FP4RecordSet Records;
 
+		if (InCommand.Changelist.IsInitialized())
+		{
+			FString ChangelistNumber = InCommand.Changelist.ToString();
+
+			if (!ChangelistNumber.IsEmpty())
+			{
+				Parameters.Add(TEXT("-c"));
+				Parameters.Add(ChangelistNumber);
+				InChangelist = InCommand.Changelist;
+			}
+		}
+
 		AppendChangelistParameter(Parameters);
 		Parameters.Append(InCommand.Files);
 
@@ -643,6 +660,8 @@ bool FPerforceMarkForAddWorker::Execute(FPerforceSourceControlCommand& InCommand
 
 bool FPerforceMarkForAddWorker::UpdateStates() const
 {
+	UpdateChangelistState(InChangelist, OutResults);
+
 	return UpdateCachedStates(OutResults);
 }
 
@@ -756,9 +775,9 @@ static void ParseSyncResults(const FP4RecordSet& InRecords, TMap<FString, EPerfo
 		FString FullPath(FileName);
 		FPaths::NormalizeFilename(FullPath);
 
-		if(Action.Len() > 0)
+		if (Action.Len() > 0)
 		{
-			if(Action == TEXT("updated"))
+			if (Action == TEXT("updated") || (Action == TEXT("refreshed")))
 			{
 				OutResults.Add(FullPath, EPerforceState::ReadOnly);
 			}
@@ -772,24 +791,35 @@ bool FPerforceSyncWorker::Execute(FPerforceSourceControlCommand& InCommand)
 	if (!InCommand.IsCanceled() && ScopedConnection.IsValid())
 	{
 		FPerforceConnection& Connection = ScopedConnection.GetConnection();
+		TSharedRef<FSync, ESPMode::ThreadSafe> Operation = StaticCastSharedRef<FSync>(InCommand.Operation);
 		TArray<FString> Parameters;
-		Parameters.Append(InCommand.Files);
 
-		TSharedRef<FSync, ESPMode::ThreadSafe> Operation = StaticCastSharedRef<FSync>(InCommand.Operation);
+		if (Operation->IsForced())
+		{
+			Parameters.Add(TEXT("-f"));
+		}
+
 		const FString& Revision = Operation->GetRevision();
 
 		// check for directories and add '...'
-		for(FString& FileName : Parameters)
+		for (FString FileName : InCommand.Files)
 		{
-			if(FileName.EndsWith(TEXT("/")))
+			if (FileName.EndsWith(TEXT("/")))
 			{
 				FileName += TEXT("...");
 			}
-			if (!Revision.IsEmpty())
+
+			if (Operation->IsHeadRevisionFlagSet())
+			{
+				FileName += TEXT("#head");
+			}
+			else if (!Revision.IsEmpty())
 			{
 				// @= syncs the file to the submitted/shelved changelist number
 				FileName += FString::Printf(TEXT("@%s"), *Revision);
 			}
+
+			Parameters.Add(MoveTemp(FileName));
 		}
 
 		FP4RecordSet Records;
@@ -2246,6 +2276,8 @@ bool FPerforceNewChangelistWorker::Execute(class FPerforceSourceControlCommand&
 			NewChangelistState.Description = Operation->GetDescription().ToString();
 			NewChangelistState.bHasShelvedFiles = false;
 
+			Operation->SetNewChangelist(MakeShared<FPerforceSourceControlChangelist>(NewChangelist));
+
 			if (InCommand.Files.Num() > 0)
 			{
 				// Move files to changelist
diff --git a/Engine/Plugins/Developer/PerforceSourceControl/Source/PerforceSourceControl/Private/PerforceSourceControlOperations.h b/Engine/Plugins/Developer/PerforceSourceControl/Source/PerforceSourceControl/Private/PerforceSourceControlOperations.h
--- a/Engine/Plugins/Developer/PerforceSourceControl/Source/PerforceSourceControl/Private/PerforceSourceControlOperations.h
+++ b/Engine/Plugins/Developer/PerforceSourceControl/Source/PerforceSourceControl/Private/PerforceSourceControlOperations.h
@@ -67,6 +67,7 @@ public:
 public:
 	/** Map of filenames to perforce state */
 	TMap<FString, EPerforceState::Type> OutResults;
+	FPerforceSourceControlChangelist InChangelist;
 };
 
 class FPerforceDeleteWorker : public IPerforceSourceControlWorker
diff --git a/Engine/Source/Developer/SourceControl/Public/SourceControlOperations.h b/Engine/Source/Developer/SourceControl/Public/SourceControlOperations.h
--- a/Engine/Source/Developer/SourceControl/Public/SourceControlOperations.h
+++ b/Engine/Source/Developer/SourceControl/Public/SourceControlOperations.h
@@ -218,9 +218,35 @@ public:
 		return Revision;
 	}
 
+	void SetHeadRevisionFlag(const bool bInHeadRevision)
+	{
+		bHeadRevision = bInHeadRevision;
+	}
+
+	bool IsHeadRevisionFlagSet() const
+	{
+		return bHeadRevision;
+	}
+
+	void SetForce(const bool bInForce)
+	{
+		bForce = bInForce;
+	}
+
+	bool IsForced() const
+	{
+		return bForce;
+	}
+
 protected:
 	/** Revision to sync to */
 	FString Revision;
+
+	/** Flag abstracting if the operation aim to sync to head */
+	bool bHeadRevision = false;
+
+	/** Forces operation, even if the file is already at the wanted revision. */
+	bool bForce = false;
 };
 
 /**
@@ -477,9 +503,20 @@ public:
 		return Description;
 	}
 
+	void SetNewChangelist(FSourceControlChangelistPtr InNewChangelist)
+	{
+		NewChangelist = InNewChangelist;
+	}
+
+	const FSourceControlChangelistPtr& GetNewChangelist() const
+	{
+		return NewChangelist;
+	}
+
 protected:
 	/** Description of the changelist */
 	FText Description;
+	FSourceControlChangelistPtr NewChangelist;
 };
 
 /**
diff --git a/Engine/Source/Developer/UncontrolledChangelists/Private/UncontrolledChangelistsModule.cpp b/Engine/Source/Developer/UncontrolledChangelists/Private/UncontrolledChangelistsModule.cpp
--- a/Engine/Source/Developer/UncontrolledChangelists/Private/UncontrolledChangelistsModule.cpp
+++ b/Engine/Source/Developer/UncontrolledChangelists/Private/UncontrolledChangelistsModule.cpp
@@ -263,20 +263,37 @@ void FUncontrolledChangelistsModule::MoveFilesToControlledChangelist(const TArra
 		return;
 	}
 
-	ISourceControlProvider& SourceControlProvider = ISourceControlModule::Get().GetProvider();
 	TArray<FString> UncontrolledFilenames;
+	
+	Algo::Transform(InUncontrolledFileStates, UncontrolledFilenames, [](const FSourceControlStateRef& State) { return State->GetFilename(); });
+	MoveFilesToControlledChangelist(UncontrolledFilenames, InChangelist);
+}
+
+void FUncontrolledChangelistsModule::MoveFilesToControlledChangelist(const TArray<FString>& InUncontrolledFiles, const FSourceControlChangelistPtr& InChangelist)
+{
+	if (!IsEnabled())
+	{
+		return;
+	}
+
+	ISourceControlProvider& SourceControlProvider = ISourceControlModule::Get().GetProvider();
 	TArray<FSourceControlStateRef> UpdatedFilestates;
 
 	// Get updated filestates to check Checkout capabilities.
-	Algo::Transform(InUncontrolledFileStates, UncontrolledFilenames, [](const FSourceControlStateRef& State) { return State->GetFilename(); });
-	SourceControlProvider.GetState(UncontrolledFilenames, UpdatedFilestates, EStateCacheUsage::ForceUpdate);
+	SourceControlProvider.GetState(InUncontrolledFiles, UpdatedFilestates, EStateCacheUsage::ForceUpdate);
 
 	TArray<UPackage*> PackageConflicts;
+	TArray<FString> FilesToAdd;
+	TArray<FString> FilesToCheckout;
 
-	// Check if we can Checkout files
+	// Check if we can Checkout files or mark for add
 	for (const FSourceControlStateRef& Filestate : UpdatedFilestates)
 	{
-		if (!Filestate->CanCheckout())
+		if (!Filestate->IsSourceControlled())
+		{
+			FilesToAdd.Add(Filestate->GetFilename());
+		}
+		else if (!Filestate->CanCheckout())
 		{
 			FString PackageName;
 			FPackageName::TryConvertFilenameToLongPackageName(Filestate->GetFilename(), PackageName);
@@ -285,8 +302,13 @@ void FUncontrolledChangelistsModule::MoveFilesToControlledChangelist(const TArra
 			if (FoundPackage != nullptr)
 			{
 				PackageConflicts.Add(FoundPackage);
+				FilesToCheckout.Add(Filestate->GetFilename());
 			}
 		}
+		else
+		{
+			FilesToCheckout.Add(Filestate->GetFilename());
+		}
 	}
 
 	bool bCanProceed = true;
@@ -299,7 +321,15 @@ void FUncontrolledChangelistsModule::MoveFilesToControlledChangelist(const TArra
 
 	if (bCanProceed)
 	{
-		SourceControlProvider.Execute(ISourceControlOperation::Create<FCheckOut>(), InChangelist, UncontrolledFilenames);
+		if (!FilesToCheckout.IsEmpty())
+		{
+			SourceControlProvider.Execute(ISourceControlOperation::Create<FCheckOut>(), InChangelist, FilesToCheckout);
+		}
+
+		if (!FilesToAdd.IsEmpty())
+		{
+			SourceControlProvider.Execute(ISourceControlOperation::Create<FMarkForAdd>(), InChangelist, FilesToAdd);
+		}
 
 		// UpdateStatus so UncontrolledChangelists can remove files from their cache if they were present before checkout.
 		UpdateStatus();
diff --git a/Engine/Source/Developer/UncontrolledChangelists/Public/UncontrolledChangelistState.h b/Engine/Source/Developer/UncontrolledChangelists/Public/UncontrolledChangelistState.h
--- a/Engine/Source/Developer/UncontrolledChangelists/Public/UncontrolledChangelistState.h
+++ b/Engine/Source/Developer/UncontrolledChangelists/Public/UncontrolledChangelistState.h
@@ -123,4 +123,5 @@ public:
 };
 
 ENUM_CLASS_FLAGS(FUncontrolledChangelistState::ECheckFlags);
+typedef TSharedPtr<FUncontrolledChangelistState, ESPMode::ThreadSafe> FUncontrolledChangelistStatePtr;
 typedef TSharedRef<FUncontrolledChangelistState, ESPMode::ThreadSafe> FUncontrolledChangelistStateRef;
diff --git a/Engine/Source/Developer/UncontrolledChangelists/Public/UncontrolledChangelistsModule.h b/Engine/Source/Developer/UncontrolledChangelists/Public/UncontrolledChangelistsModule.h
--- a/Engine/Source/Developer/UncontrolledChangelists/Public/UncontrolledChangelistsModule.h
+++ b/Engine/Source/Developer/UncontrolledChangelists/Public/UncontrolledChangelistsModule.h
@@ -121,6 +121,13 @@ public:
 	 * @param 	InChangelist 				The Controlled Changelist where to move the files.
 	 */
 	void MoveFilesToControlledChangelist(const TArray<FSourceControlStateRef>& InUncontrolledFileStates, const FSourceControlChangelistPtr& InChangelist);
+	
+	/**
+	 * Moves files to a Controlled Changelist.
+	 * @param 	InUncontrolledFiles 	The files to move.
+	 * @param 	InChangelist 			The Controlled Changelist where to move the files.
+	 */
+	void MoveFilesToControlledChangelist(const TArray<FString>& InUncontrolledFiles, const FSourceControlChangelistPtr& InChangelist);
 
 private:
 	/**
diff --git a/Engine/Source/Editor/SourceControlWindows/Private/SSourceControlChangelists.cpp b/Engine/Source/Editor/SourceControlWindows/Private/SSourceControlChangelists.cpp
--- a/Engine/Source/Editor/SourceControlWindows/Private/SSourceControlChangelists.cpp
+++ b/Engine/Source/Editor/SourceControlWindows/Private/SSourceControlChangelists.cpp
@@ -437,6 +437,25 @@ FSourceControlChangelistStatePtr SSourceControlChangelistsWidget::GetCurrentChan
 	}
 }
 
+FUncontrolledChangelistStatePtr SSourceControlChangelistsWidget::GetCurrentUncontrolledChangelistState()
+{
+	if (!TreeView)
+	{
+		return nullptr;
+	}
+
+	TArray<FChangelistTreeItemPtr> SelectedItems = TreeView->GetSelectedItems();
+
+	if (SelectedItems.Num() != 1 || SelectedItems[0]->GetTreeItemType() != IChangelistTreeItem::UncontrolledChangelist)
+	{
+		return nullptr;
+	}
+	else
+	{
+		return StaticCastSharedPtr<FUncontrolledChangelistTreeItem>(SelectedItems[0])->UncontrolledChangelistState;
+	}
+}
+
 FSourceControlChangelistPtr SSourceControlChangelistsWidget::GetCurrentChangelist()
 {
 	FSourceControlChangelistStatePtr ChangelistState = GetCurrentChangelistState();
@@ -487,8 +506,14 @@ TArray<FString> SSourceControlChangelistsWidget::GetSelectedFiles()
 	else
 	{
 		TArray<FString> Files;
+
 		for (FChangelistTreeItemPtr Item : SelectedItems)
 		{
+			if (Item->GetTreeItemType() != IChangelistTreeItem::File)
+			{
+				continue;
+			}
+
 			Files.Add(StaticCastSharedPtr<FFileTreeItem>(Item)->FileState->GetFilename());
 		}
 
@@ -496,6 +521,68 @@ TArray<FString> SSourceControlChangelistsWidget::GetSelectedFiles()
 	}
 }
 
+void SSourceControlChangelistsWidget::GetSelectedFiles(TArray<FString>& OutControlledFiles, TArray<FString>& OutUncontrolledFiles)
+{
+	TArray<FChangelistTreeItemPtr> SelectedItems = TreeView->GetSelectedItems();
+
+	for (const FChangelistTreeItemPtr& Item : SelectedItems)
+	{
+		if (Item->GetTreeItemType() != IChangelistTreeItem::File)
+		{
+			continue;
+		}
+
+		const FChangelistTreeItemPtr& Parent = Item->GetParent();
+
+		if (!Parent.IsValid())
+		{
+			continue;
+		}
+
+		const FString& Filename = StaticCastSharedPtr<FFileTreeItem>(Item)->FileState->GetFilename();
+
+		if (Parent->GetTreeItemType() == IChangelistTreeItem::Changelist)
+		{
+			OutControlledFiles.Add(Filename);
+		}
+		else if (Parent->GetTreeItemType() == IChangelistTreeItem::UncontrolledChangelist)
+		{
+			OutUncontrolledFiles.Add(Filename);
+		}
+	}
+}
+
+void SSourceControlChangelistsWidget::GetSelectedFileStates(TArray<FSourceControlStateRef>& OutControlledFileStates, TArray<FSourceControlStateRef>& OutUncontrolledFileStates)
+{
+	TArray<FChangelistTreeItemPtr> SelectedItems = TreeView->GetSelectedItems();
+
+	for (const FChangelistTreeItemPtr& Item : SelectedItems)
+	{
+		if (Item->GetTreeItemType() != IChangelistTreeItem::File)
+		{
+			continue;
+		}
+
+		const FChangelistTreeItemPtr& Parent = Item->GetParent();
+
+		if (!Parent.IsValid())
+		{
+			continue;
+		}
+
+		FSourceControlStateRef FileState = StaticCastSharedPtr<FFileTreeItem>(Item)->FileState;
+
+		if (Parent->GetTreeItemType() == IChangelistTreeItem::Changelist)
+		{
+			OutControlledFileStates.Add(MoveTemp(FileState));
+		}
+		else if (Parent->GetTreeItemType() == IChangelistTreeItem::UncontrolledChangelist)
+		{
+			OutUncontrolledFileStates.Add(MoveTemp(FileState));
+		}
+	}
+}
+
 TArray<FString> SSourceControlChangelistsWidget::GetSelectedShelvedFiles()
 {
 	TArray<FString> ShelvedFiles;
@@ -524,6 +611,36 @@ TArray<FString> SSourceControlChangelistsWidget::GetSelectedShelvedFiles()
 	return ShelvedFiles;
 }
 
+bool SSourceControlChangelistsWidget::IsParentOfSelection(const IChangelistTreeItem::TreeItemType ParentType) const
+{
+	return Algo::AnyOf(TreeView->GetSelectedItems(), [ParentType = ParentType](const FChangelistTreeItemPtr& Item)
+	{
+		IChangelistTreeItem::TreeItemType ItemType = Item->GetTreeItemType();
+
+		if (ItemType == ParentType)
+		{
+			return true;
+		}
+		else if ((ItemType == IChangelistTreeItem::File) || (ItemType == IChangelistTreeItem::ShelvedChangelist))
+		{
+			const FChangelistTreeItemPtr& Parent = Item->GetParent();
+			return Parent.IsValid() && (Parent->GetTreeItemType() == ParentType);
+		}
+		else if (ItemType == IChangelistTreeItem::ShelvedFile)
+		{
+			const FChangelistTreeItemPtr& Parent = Item->GetParent();
+
+			if (Parent.IsValid())
+			{
+				const FChangelistTreeItemPtr& GrandParent = Parent->GetParent();
+				return GrandParent.IsValid() && (GrandParent->GetTreeItemType() == ParentType);
+			}
+		}
+
+		return false;
+	});
+}
+
 void SSourceControlChangelistsWidget::OnNewChangelist()
 {
 	FText ChangelistDescription;
@@ -625,14 +742,57 @@ void SSourceControlChangelistsWidget::OnRevert()
 		return;
 	}
 
+	TArray<FString> SelectedControlledFiles;
+	TArray<FString> SelectedUncontrolledFiles;
+
+	GetSelectedFiles(SelectedControlledFiles, SelectedUncontrolledFiles);
+
+	FSourceControlChangelistPtr SelectedChangelist = GetChangelistFromSelection();
 	ISourceControlProvider& SourceControlProvider = ISourceControlModule::Get().GetProvider();
-	auto RevertOperation = ISourceControlOperation::Create<FRevert>();
-	SourceControlProvider.Execute(RevertOperation, GetChangelistFromSelection(), GetSelectedFiles());
+
+	// Reverts the selected Changelist or Files
+	if (SelectedChangelist.IsValid() || (!SelectedControlledFiles.IsEmpty()))
+	{
+		auto RevertOperation = ISourceControlOperation::Create<FRevert>();
+		SourceControlProvider.Execute(RevertOperation, SelectedChangelist, SelectedControlledFiles);
+	}
+
+	// Reverts selected Uncontrolled Files
+	if (!SelectedUncontrolledFiles.IsEmpty())
+	{
+		auto ForceSyncOperation = ISourceControlOperation::Create<FSync>();
+		ForceSyncOperation->SetForce(true);
+		ForceSyncOperation->SetHeadRevisionFlag(true);
+		SourceControlProvider.Execute(ForceSyncOperation, SelectedUncontrolledFiles);
+
+		FUncontrolledChangelistsModule::Get().UpdateStatus();
+	}
+
+	FUncontrolledChangelistStatePtr SelectedUncontrolledChangelist = GetCurrentUncontrolledChangelistState();
+
+	// Reverts the selected Uncontrolled Changelist
+	if (SelectedUncontrolledChangelist.IsValid())
+	{
+		SelectedUncontrolledFiles.Reset();
+		Algo::Transform(SelectedUncontrolledChangelist->GetFilesStates(), SelectedUncontrolledFiles, [](const FSourceControlStateRef& State) { return State->GetFilename(); });
+
+		auto ForceSyncOperation = ISourceControlOperation::Create<FSync>();
+		ForceSyncOperation->SetForce(true);
+		ForceSyncOperation->SetHeadRevisionFlag(true);
+		SourceControlProvider.Execute(ForceSyncOperation, SelectedUncontrolledFiles);
+
+		FUncontrolledChangelistsModule::Get().UpdateStatus();
+	}
 }
 
 bool SSourceControlChangelistsWidget::CanRevert()
 {
-	return GetSelectedFiles().Num() > 0 || (GetCurrentChangelistState() && GetCurrentChangelistState()->GetFilesStates().Num() > 0);
+	FSourceControlChangelistStatePtr CurrentChangelistState = GetCurrentChangelistState();
+	FUncontrolledChangelistStatePtr CurrentUncontrolledChangelistState = GetCurrentUncontrolledChangelistState();
+
+	return GetSelectedFiles().Num() > 0
+		|| (CurrentChangelistState.IsValid() && CurrentChangelistState->GetFilesStates().Num() > 0)
+		|| (CurrentUncontrolledChangelistState.IsValid() && CurrentUncontrolledChangelistState->GetFilesStates().Num() > 0);
 }
 
 void SSourceControlChangelistsWidget::OnShelve()
@@ -864,9 +1024,12 @@ bool SSourceControlChangelistsWidget::CanSubmitChangelist()
 
 void SSourceControlChangelistsWidget::OnMoveFiles()
 {
-	TArray<FString> SelectedFiles = GetSelectedFiles();
+	TArray<FString> SelectedControlledFiles;
+	TArray<FString> SelectedUncontrolledFiles;
+	
+	GetSelectedFiles(SelectedControlledFiles, SelectedUncontrolledFiles);
 
-	if (SelectedFiles.Num() == 0)
+	if (SelectedControlledFiles.IsEmpty() && SelectedUncontrolledFiles.IsEmpty())
 	{
 		return;
 	}
@@ -890,13 +1053,21 @@ void SSourceControlChangelistsWidget::OnMoveFiles()
 
 	for (FChangelistTreeItemPtr Changelist : ChangelistsNodes)
 	{
-		if (!Changelist || Changelist->GetTreeItemType() != IChangelistTreeItem::Changelist)
+		if (!Changelist)
 		{
 			continue;
 		}
 
-		auto TypedChangelist = StaticCastSharedPtr<FChangelistTreeItem>(Changelist);
-		Items.Emplace(TypedChangelist->GetDisplayText(), TypedChangelist->GetDescriptionText(), bCanEditAlreadyExistingChangelistDescription);
+		if (Changelist->GetTreeItemType() == IChangelistTreeItem::Changelist)
+		{
+			const auto& TypedChangelist = StaticCastSharedPtr<FChangelistTreeItem>(Changelist);
+			Items.Emplace(TypedChangelist->GetDisplayText(), TypedChangelist->GetDescriptionText(), bCanEditAlreadyExistingChangelistDescription);
+		}
+		else if (Changelist->GetTreeItemType() == IChangelistTreeItem::UncontrolledChangelist)
+		{
+			const FUncontrolledChangelistTreeItemPtr& TypedChangelist = StaticCastSharedPtr<FUncontrolledChangelistTreeItem>(Changelist);
+			Items.Emplace(TypedChangelist->GetDisplayText(), TypedChangelist->GetDescriptionText(), bCanEditAlreadyExistingChangelistDescription);
+		}
 	}
 
 	int32 PickedItem = 0;
@@ -922,13 +1093,46 @@ void SSourceControlChangelistsWidget::OnMoveFiles()
 	{
 		auto NewChangelistOperation = ISourceControlOperation::Create<FNewChangelist>();
 		NewChangelistOperation->SetDescription(ChangelistDescription);
-		SourceControlProvider.Execute(NewChangelistOperation, SelectedFiles);
+		SourceControlProvider.Execute(NewChangelistOperation, SelectedControlledFiles);
+
+		if ((!SelectedUncontrolledFiles.IsEmpty()) && NewChangelistOperation->GetNewChangelist().IsValid())
+		{
+			FUncontrolledChangelistsModule::Get().MoveFilesToControlledChangelist(SelectedUncontrolledFiles, NewChangelistOperation->GetNewChangelist());
+		}
 	}
 	else
 	{
 		const int32 ChangelistIndex = (bAddNewChangelistEntry ? PickedItem - 1 : PickedItem);
-		FSourceControlChangelistPtr Changelist = StaticCastSharedPtr<FChangelistTreeItem>(ChangelistsNodes[ChangelistIndex])->ChangelistState->GetChangelist();
-		SourceControlProvider.Execute(ISourceControlOperation::Create<FMoveToChangelist>(), Changelist, SelectedFiles);
+		const FChangelistTreeItemPtr& SelectedItem = ChangelistsNodes[ChangelistIndex];
+
+		if (SelectedItem->GetTreeItemType() == IChangelistTreeItem::Changelist)
+		{
+			FSourceControlChangelistPtr Changelist = StaticCastSharedPtr<FChangelistTreeItem>(SelectedItem)->ChangelistState->GetChangelist();
+
+			if (!SelectedControlledFiles.IsEmpty())
+			{
+				SourceControlProvider.Execute(ISourceControlOperation::Create<FMoveToChangelist>(), Changelist, SelectedControlledFiles);
+			}
+
+			if (!SelectedUncontrolledFiles.IsEmpty())
+			{
+				FUncontrolledChangelistsModule::Get().MoveFilesToControlledChangelist(SelectedUncontrolledFiles, Changelist);
+			}
+		}
+		else if (SelectedItem->GetTreeItemType() == IChangelistTreeItem::UncontrolledChangelist)
+		{
+			const FUncontrolledChangelist UncontrolledChangelist = StaticCastSharedPtr<FUncontrolledChangelistTreeItem>(SelectedItem)->UncontrolledChangelistState->Changelist;
+			
+			TArray<FSourceControlStateRef> SelectedControlledFileStates;
+			TArray<FSourceControlStateRef> SelectedUnControlledFileStates;
+
+			GetSelectedFileStates(SelectedControlledFileStates, SelectedUnControlledFileStates);
+
+			if ((!SelectedControlledFileStates.IsEmpty()) || (!SelectedUnControlledFileStates.IsEmpty()))
+			{
+				FUncontrolledChangelistsModule::Get().MoveFilesToUncontrolledChangelist(SelectedControlledFileStates, SelectedUnControlledFileStates, UncontrolledChangelist);
+			}
+		}
 	}
 }
 
@@ -1014,11 +1218,13 @@ TSharedPtr<SWidget> SSourceControlChangelistsWidget::OnOpenContextMenu()
 	// Build up the menu for a selection
 	FToolMenuContext Context;
 	UToolMenu* Menu = ToolMenus->GenerateMenu(MenuName, Context);
-		
+
 	bool bHasSelectedChangelist = (GetCurrentChangelist() != nullptr);
 	bool bHasSelectedFiles = (GetSelectedFiles().Num() > 0);
 	bool bHasSelectedShelvedFiles = (GetSelectedShelvedFiles().Num() > 0);
 	bool bHasEmptySelection = (!bHasSelectedChangelist && !bHasSelectedFiles && !bHasSelectedShelvedFiles);
+	bool bIsChangelistParentOfSelection = IsParentOfSelection(IChangelistTreeItem::Changelist);
+	bool bIsUncontrolledChangelistParentOfSelection = IsParentOfSelection(IChangelistTreeItem::UncontrolledChangelist);
 
 	FToolMenuSection& Section = Menu->AddSection("Source Control");
 	
@@ -1032,25 +1238,27 @@ TSharedPtr<SWidget> SSourceControlChangelistsWidget::OnOpenContextMenu()
 	}
 
 	// This can appear on both files & changelist
-	if (bHasSelectedChangelist || bHasSelectedFiles)
+	if (bIsChangelistParentOfSelection)
 	{
 		Section.AddMenuEntry("RevertUnchanged", LOCTEXT("SourceControl_RevertUnchanged", "Revert Unchanged"), LOCTEXT("SourceControl_Revert_Unchanged_Tooltip", "Reverts unchanged files & changelists"), FSlateIcon(),
-			FUIAction(
-				FExecuteAction::CreateSP(this, &SSourceControlChangelistsWidget::OnRevertUnchanged),
-				FCanExecuteAction::CreateSP(this, &SSourceControlChangelistsWidget::CanRevertUnchanged)));
+							 FUIAction(FExecuteAction::CreateSP(this, &SSourceControlChangelistsWidget::OnRevertUnchanged),
+									   FCanExecuteAction::CreateSP(this, &SSourceControlChangelistsWidget::CanRevertUnchanged)));
+	}
 
+	if (bIsChangelistParentOfSelection || bIsUncontrolledChangelistParentOfSelection)
+	{
 		Section.AddMenuEntry("Revert", LOCTEXT("SourceControl_Revert", "Revert Files"), LOCTEXT("SourceControl_Revert_Tooltip", "Reverts all files in the changelist or from the selection"), FSlateIcon(),
 			FUIAction(
 				FExecuteAction::CreateSP(this, &SSourceControlChangelistsWidget::OnRevert),
 				FCanExecuteAction::CreateSP(this, &SSourceControlChangelistsWidget::CanRevert)));
 	}
 
-	if(bHasSelectedFiles || bHasSelectedShelvedFiles || (bHasSelectedChangelist && (GetCurrentChangelistState()->GetFilesStates().Num() > 0 || GetCurrentChangelistState()->GetShelvedFilesStates().Num() > 0)))
+	if (bIsChangelistParentOfSelection && (bHasSelectedFiles || bHasSelectedShelvedFiles || (bHasSelectedChangelist && (GetCurrentChangelistState()->GetFilesStates().Num() > 0 || GetCurrentChangelistState()->GetShelvedFilesStates().Num() > 0))))
 	{
 		Section.AddSeparator("ShelveSeparator");
 	}
 
-	if(bHasSelectedFiles || (bHasSelectedChangelist && GetCurrentChangelistState()->GetFilesStates().Num() > 0))
+	if (bIsChangelistParentOfSelection && (bHasSelectedFiles || (bHasSelectedChangelist && GetCurrentChangelistState()->GetFilesStates().Num() > 0)))
 	{
 		Section.AddMenuEntry("Shelve", LOCTEXT("SourceControl_Shelve", "Shelve Files"), LOCTEXT("SourceControl_Shelve_Tooltip", "Shelves the changelist or the selected files"), FSlateIcon(),
 			FUIAction(FExecuteAction::CreateSP(this, &SSourceControlChangelistsWidget::OnShelve)));
diff --git a/Engine/Source/Editor/SourceControlWindows/Private/SSourceControlChangelists.h b/Engine/Source/Editor/SourceControlWindows/Private/SSourceControlChangelists.h
--- a/Engine/Source/Editor/SourceControlWindows/Private/SSourceControlChangelists.h
+++ b/Engine/Source/Editor/SourceControlWindows/Private/SSourceControlChangelists.h
@@ -46,6 +46,7 @@ private:
 
 	/** Returns the currently selected changelist state ptr or null in invalid cases */
 	FSourceControlChangelistStatePtr GetCurrentChangelistState();
+	FUncontrolledChangelistStatePtr GetCurrentUncontrolledChangelistState();
 	FSourceControlChangelistPtr GetCurrentChangelist();
 	FSourceControlChangelistStatePtr GetChangelistStateFromSelection();
 	FSourceControlChangelistPtr GetChangelistFromSelection();
@@ -53,9 +54,30 @@ private:
 	/** Returns list of currently selected files */
 	TArray<FString> GetSelectedFiles();
 
+	/**
+	 * Splits selected files between Controlled and Uncontrolled files.
+	 * @param 	OutControlledFiles 		Selected source controlled files will be added to this array.
+	 * @param 	OutUncontrolledFiles	Selected uncontrolled files will be added to this array.
+	 */
+	void GetSelectedFiles(TArray<FString>& OutControlledFiles, TArray<FString>& OutUncontrolledFiles);
+
+	/**
+	 * Splits selected files between Controlled and Uncontrolled files.
+	 * @param 	OutControlledFileStates 	Selected source controlled file states will be added to this array.
+	 * @param 	OutUncontrolledFileStates	Selected uncontrolled file states will be added to this array.
+	 */
+	void GetSelectedFileStates(TArray<FSourceControlStateRef>& OutControlledFileStates, TArray<FSourceControlStateRef>& OutUncontrolledFileStates);
+
 	/** Returns list of currently selected shelved files */
 	TArray<FString> GetSelectedShelvedFiles();
 
+	/**
+	 * Check if the type given as argument is a parent of selected items.
+	 * @param 	ParentType 	The parent type to look for.
+	 * @return 	True of ParentType is a parent of selected items.
+	 */
+	bool IsParentOfSelection(const IChangelistTreeItem::TreeItemType ParentType) const;
+
 	/** Changelist operations */
 	void OnNewChangelist();
 	void OnDeleteChangelist();
