* Adds "Validate Changelist" contextual action in Source Control Changelists window.
* Moves validation logs to message logs so they can be both accessible in OutputLogs and Hyperlink to SourceControl MessageLog

#changelist validated

diff --git a/Engine/Source/Editor/SourceControlWindows/Private/SSourceControlChangelists.cpp b/Engine/Source/Editor/SourceControlWindows/Private/SSourceControlChangelists.cpp
--- a/Engine/Source/Editor/SourceControlWindows/Private/SSourceControlChangelists.cpp
+++ b/Engine/Source/Editor/SourceControlWindows/Private/SSourceControlChangelists.cpp
@@ -833,7 +833,7 @@ void SSourceControlChangelistsWidget::OnDeleteShelvedFiles()
 	SourceControlProvider.Execute(DeleteShelvedOperation, GetChangelistFromSelection(), GetSelectedShelvedFiles());
 }
 
-static bool GetChangelistValidationResult(FSourceControlChangelistPtr InChangelist, FString& OutValidationText, FName& OutValidationIcon)
+static bool GetChangelistValidationResult(FSourceControlChangelistPtr InChangelist, FString& OutValidationText)
 {
 	FSourceControlPreSubmitDataValidationDelegate ValidationDelegate = ISourceControlModule::Get().GetRegisteredPreSubmitDataValidation();
 
@@ -845,24 +845,28 @@ static bool GetChangelistValidationResult(FSourceControlChangelistPtr InChangeli
 
 	if (ValidationDelegate.ExecuteIfBound(InChangelist, ValidationResult, ValidationErrors, ValidationWarnings))
 	{
+		EMessageSeverity::Type MessageSeverity = EMessageSeverity::Info;
+
 		if (ValidationResult == EDataValidationResult::Invalid || ValidationErrors.Num() > 0)
 		{
-			OutValidationIcon = "MessageLog.Error";
 			OutValidationText = LOCTEXT("SourceControl.Submit.ChangelistValidationError", "Changelist validation failed!").ToString();
 			bValidationResult = false;
+			MessageSeverity = EMessageSeverity::Error;
 		}
 		else if (ValidationResult == EDataValidationResult::NotValidated || ValidationWarnings.Num() > 0)
 		{
-			OutValidationIcon = "MessageLog.Warning";
 			OutValidationText = LOCTEXT("SourceControl.Submit.ChangelistValidationWarning", "Changelist validation has warnings!").ToString();
+			MessageSeverity = EMessageSeverity::Warning;
 		}
 		else
 		{
-			OutValidationIcon = "MessageLog.Note";
-			OutValidationText = LOCTEXT("SourceControl.Submit.ChangelistValidationSuccess", "Changelist validation succesful!").ToString();
+			OutValidationText = LOCTEXT("SourceControl.Submit.ChangelistValidationSuccess", "Changelist validation successful!").ToString();
 		}
 
 		int32 NumLinesDisplayed = 0;
+		FMessageLog SourceControlLog("SourceControl");
+		
+		SourceControlLog.Message(MessageSeverity, FText::FromString(*OutValidationText));
 
 		auto AppendInfo = [&OutValidationText, &NumLinesDisplayed](const TArray<FText>& Info, const FString& InfoType)
 		{
@@ -890,14 +894,15 @@ static bool GetChangelistValidationResult(FSourceControlChangelistPtr InChangeli
 			}
 		};
 
-		auto LogInfo = [](const TArray<FText>& Info, const FString& InfoType, const ELogVerbosity::Type LogVerbosity)
+		auto LogInfo = [&SourceControlLog](const TArray<FText>& Info, const FString& InfoType, const EMessageSeverity::Type LogVerbosity)
 		{
 			if (Info.Num() > 0)
 			{
-				FMsg::Logf(nullptr, 0, LogSourceControlChangelist.GetCategoryName(), LogVerbosity, TEXT("Encountered %d %s:"), Info.Num(), *InfoType);
+				SourceControlLog.Message(LogVerbosity, FText::Format(LOCTEXT("SourceControl.Validation.ErrorEncountered", "Encountered {0} {1}:"), FText::AsNumber(Info.Num()), FText::FromString(*InfoType)));
+
 				for (const FText& Line : Info)
 				{
-					FMsg::Logf(nullptr, 0, LogSourceControlChangelist.GetCategoryName(), LogVerbosity, TEXT("%s"), *Line.ToString());
+					SourceControlLog.Message(LogVerbosity, Line);
 				}
 			}
 		};
@@ -905,8 +910,8 @@ static bool GetChangelistValidationResult(FSourceControlChangelistPtr InChangeli
 		AppendInfo(ValidationErrors, TEXT("errors"));
 		AppendInfo(ValidationWarnings, TEXT("warnings"));
 
-		LogInfo(ValidationErrors, TEXT("errors"), ELogVerbosity::Error);
-		LogInfo(ValidationWarnings, TEXT("warnings"), ELogVerbosity::Warning);
+		LogInfo(ValidationErrors, TEXT("errors"), EMessageSeverity::Error);
+		LogInfo(ValidationWarnings, TEXT("warnings"), EMessageSeverity::Warning);
 	}
 
 	return bValidationResult;
@@ -922,13 +927,13 @@ void SSourceControlChangelistsWidget::OnSubmitChangelist()
 	}
 
 	FString ChangelistValidationText;
-	FName ChangelistValidationIconName;
-	bool bValidationResult = GetChangelistValidationResult(ChangelistState->GetChangelist(), ChangelistValidationText, ChangelistValidationIconName);
+	bool bValidationResult = GetChangelistValidationResult(ChangelistState->GetChangelist(), ChangelistValidationText);
 
 	// Build list of states for the dialog
 	const FText OriginalChangelistDescription = ChangelistState->GetDescriptionText();
 	const bool bAskForChangelistDescription = (OriginalChangelistDescription.IsEmptyOrWhitespace());
 	FText ChangelistDescriptionToSubmit = UpdateChangelistDescriptionToSubmitIfNeeded(bValidationResult, OriginalChangelistDescription);
+	FName ChangelistValidationIconName = bValidationResult ? TEXT("Icons.SuccessWithColor.Large") : TEXT("Icons.ErrorWithColor.Large");
 
 	TSharedRef<SWindow> NewWindow = SNew(SWindow)
 		.Title(NSLOCTEXT("SourceControl.ConfirmSubmit", "Title", "Confirm changelist submit"))
@@ -979,14 +984,14 @@ void SSourceControlChangelistsWidget::OnSubmitChangelist()
 		if (!bCheckinSuccess)
 		{
 			Info.Text = LOCTEXT("SCC_Checkin_Failed", "Failed to check in files!");
-			Info.Image = FCoreStyle::Get().GetBrush(TEXT("MessageLog.Error"));
 		}
 		
 		Info.ExpireDuration = 8.0f;
 		Info.HyperlinkText = LOCTEXT("SCC_Checkin_ShowLog", "Show Message Log");
 		Info.Hyperlink = FSimpleDelegate::CreateLambda([]() { FMessageLog("SourceControl").Open(EMessageSeverity::Info, true); });
 
-		FSlateNotificationManager::Get().AddNotification(Info);
+		TSharedPtr<SNotificationItem> Notification = FSlateNotificationManager::Get().AddNotification(Info);
+		Notification->SetCompletionState(bCheckinSuccess ? SNotificationItem::CS_Success : SNotificationItem::CS_Fail);
 	}
 }
 
@@ -1014,6 +1019,41 @@ bool SSourceControlChangelistsWidget::CanSubmitChangelist()
 	return Changelist != nullptr && Changelist->GetFilesStates().Num() > 0 && Changelist->GetShelvedFilesStates().Num() == 0;
 }
 
+void SSourceControlChangelistsWidget::OnValidateChangelist()
+{
+	FSourceControlChangelistStatePtr ChangelistState = GetCurrentChangelistState();
+
+	if (!ChangelistState)
+	{
+		return;
+	}
+
+	FString ChangelistValidationText;
+	bool bValidationResult = GetChangelistValidationResult(ChangelistState->GetChangelist(), ChangelistValidationText);
+
+	// Setup the notification for operation feedback
+	FNotificationInfo Info(LOCTEXT("SCC_Validation_Success", "Changelist validated"));
+
+	// Override the notification fields for failure ones
+	if (!bValidationResult)
+	{
+		Info.Text = LOCTEXT("SCC_Validation_Failed", "Failed to validate the changelist");
+	}
+
+	Info.ExpireDuration = 8.0f;
+	Info.HyperlinkText = LOCTEXT("SCC_Validation_ShowLog", "Show Message Log");
+	Info.Hyperlink = FSimpleDelegate::CreateLambda([]() { FMessageLog("SourceControl").Open(EMessageSeverity::Info, true); });
+
+	TSharedPtr<SNotificationItem> Notification = FSlateNotificationManager::Get().AddNotification(Info);
+	Notification->SetCompletionState(bValidationResult ? SNotificationItem::CS_Success : SNotificationItem::CS_Fail);
+}
+
+bool SSourceControlChangelistsWidget::CanValidateChangelist()
+{
+	FSourceControlChangelistStatePtr Changelist = GetCurrentChangelistState();
+	return Changelist != nullptr && Changelist->GetFilesStates().Num() > 0;
+}
+
 void SSourceControlChangelistsWidget::OnMoveFiles()
 {
 	TArray<FString> SelectedControlledFiles;
@@ -1227,6 +1267,11 @@ TSharedPtr<SWidget> SSourceControlChangelistsWidget::OnOpenContextMenu()
 			FUIAction(
 				FExecuteAction::CreateSP(this, &SSourceControlChangelistsWidget::OnSubmitChangelist),
 				FCanExecuteAction::CreateSP(this, &SSourceControlChangelistsWidget::CanSubmitChangelist)));
+
+		Section.AddMenuEntry("ValidateChangelist", LOCTEXT("SourceControl_ValidateChangelist", "Validate Changelist"), LOCTEXT("SourceControl_ValidateChangeslit_Tooltip", "Validates a changelist"), FSlateIcon(),
+							 FUIAction(
+							 FExecuteAction::CreateSP(this, &SSourceControlChangelistsWidget::OnValidateChangelist),
+							 FCanExecuteAction::CreateSP(this, &SSourceControlChangelistsWidget::CanValidateChangelist)));
 	}
 
 	// This can appear on both files & changelist
diff --git a/Engine/Source/Editor/SourceControlWindows/Private/SSourceControlChangelists.h b/Engine/Source/Editor/SourceControlWindows/Private/SSourceControlChangelists.h
--- a/Engine/Source/Editor/SourceControlWindows/Private/SSourceControlChangelists.h
+++ b/Engine/Source/Editor/SourceControlWindows/Private/SSourceControlChangelists.h
@@ -85,6 +85,8 @@ private:
 	void OnEditChangelist();
 	void OnSubmitChangelist();
 	bool CanSubmitChangelist();
+	void OnValidateChangelist();
+	bool CanValidateChangelist();
 
 	/** Changelist & File operations */
 	void OnRevertUnchanged();
