Fixes FFileTreeItem displaying "Unavailable" when file contains multiple AssetData.
Now aggregates AssetData names while displaying "Multiple Assets" as type, ignoring Redirectors..

#changelist validated

diff --git a/Engine/Source/Editor/SourceControlWindows/Private/SSourceControlCommon.cpp b/Engine/Source/Editor/SourceControlWindows/Private/SSourceControlCommon.cpp
--- a/Engine/Source/Editor/SourceControlWindows/Private/SSourceControlCommon.cpp
+++ b/Engine/Source/Editor/SourceControlWindows/Private/SSourceControlCommon.cpp
@@ -2,6 +2,8 @@
 
 #include "SSourceControlCommon.h"
 
+#include "Algo/Count.h"
+#include "Algo/Find.h"
 #include "AssetData.h"
 #include "AssetToolsModule.h"
 #include "EditorStyleSet.h"
@@ -89,34 +91,20 @@ void FFileTreeItem::RefreshAssetInformation()
 
 	if (Assets.IsValid() && (Assets->Num() > 0))
 	{
-		FAssetData& AssetData = (*Assets)[0];
+		auto IsNotRedirector = [](const FAssetData& InAssetData) { return !InAssetData.IsRedirector(); };
+		int32 NumUserFacingAsset = Algo::CountIf(*Assets, IsNotRedirector);
 
-		TempAssetPath = AssetData.ObjectPath.ToString();
-
-		// Strip asset name from object path
-		int32 LastDot = -1;
-		if (TempAssetPath.FindLastChar('.', LastDot))
+		if (NumUserFacingAsset == 1)
 		{
-			TempAssetPath.LeftInline(LastDot);
-		}
-
-		// Find name, asset type & color only if there is exactly one asset
-		if (Assets->Num() == 1)
-		{
-			static FName NAME_ActorLabel(TEXT("ActorLabel"));
-			if (AssetData.FindTag(NAME_ActorLabel))
-			{
-				AssetData.GetTagValue(NAME_ActorLabel, TempAssetName);
-			}
-			else
-			{
-				TempAssetName = AssetData.AssetName.ToString();
-			}
+			const FAssetData& AssetData = *Algo::FindByPredicate(*Assets, IsNotRedirector);
 
+			TempAssetName = RetrieveAssetName(AssetData);
+			TempAssetPath = RetrieveAssetPath(AssetData);
 			TempAssetType = AssetData.AssetClass.ToString();
 
 			const FAssetToolsModule& AssetToolsModule = FModuleManager::LoadModuleChecked<FAssetToolsModule>(TEXT("AssetTools"));
 			const TSharedPtr<IAssetTypeActions> AssetTypeActions = AssetToolsModule.Get().GetAssetTypeActionsForClass(AssetData.GetClass()).Pin();
+
 			if (AssetTypeActions.IsValid())
 			{
 				TempAssetColor = AssetTypeActions->GetTypeColor();
@@ -128,6 +116,14 @@ void FFileTreeItem::RefreshAssetInformation()
 		}
 		else
 		{
+			TempAssetName = RetrieveAssetName((*Assets)[0]);
+			TempAssetPath = RetrieveAssetPath((*Assets)[0]);
+
+			for (int32 i = 1; i < Assets->Num(); ++i)
+			{
+				TempAssetName += TEXT(";") + RetrieveAssetName((*Assets)[i]);
+			}
+
 			TempAssetType = SSourceControlCommon::GetDefaultMultipleAsset().ToString();
 			TempAssetColor = FColor::White;
 		}
@@ -181,6 +177,37 @@ FText FFileTreeItem::GetAssetName()
 	return AssetName;
 }
 
+FString FFileTreeItem::RetrieveAssetName(const FAssetData& InAssetData) const
+{
+	static FName NAME_ActorLabel(TEXT("ActorLabel"));
+
+	if (InAssetData.FindTag(NAME_ActorLabel))
+	{
+		FString ResultAssetName = TEXT("");
+		
+		InAssetData.GetTagValue(NAME_ActorLabel, ResultAssetName);
+		return ResultAssetName;
+	}
+	else
+	{
+		return InAssetData.AssetName.ToString();
+	}
+}
+
+FString FFileTreeItem::RetrieveAssetPath(const FAssetData& InAssetData) const
+{
+	int32 LastDot = -1;
+	FString Path = InAssetData.ObjectPath.ToString();
+
+	// Strip asset name from object path
+	if (Path.FindLastChar('.', LastDot))
+	{
+		Path.LeftInline(LastDot);
+	}
+
+	return Path;
+}
+
 //////////////////////////////////////////////////////////////////////////
 
 FOfflineFileTreeItem::FOfflineFileTreeItem(const FString& InFilename)
diff --git a/Engine/Source/Editor/SourceControlWindows/Private/SSourceControlCommon.h b/Engine/Source/Editor/SourceControlWindows/Private/SSourceControlCommon.h
--- a/Engine/Source/Editor/SourceControlWindows/Private/SSourceControlCommon.h
+++ b/Engine/Source/Editor/SourceControlWindows/Private/SSourceControlCommon.h
@@ -166,6 +166,13 @@ struct FFileTreeItem : public IChangelistTreeItem
 
 	bool IsShelved() const { return GetTreeItemType() == IChangelistTreeItem::ShelvedFile; }
 
+private:
+	/** Returns a string representing the name of the asset represented by the given AssetData */
+	FString RetrieveAssetName(const FAssetData& InAssetData) const;
+
+	/** Returns a string representing the path of the asset represented by the given AssetData */
+	FString RetrieveAssetPath(const FAssetData& InAssetData) const;
+
 public:
 	/** Shared pointer to the source control state object itself */
 	FSourceControlStateRef FileState;
