Fixes empty assets crashing Changelist view.
Fixes Perforce' "DontCare" state preventing UncontrolledChangelist from freeing "not controlled + deleted" files.

#changelist validated

diff --git a/Engine/Source/Developer/UncontrolledChangelists/Private/UncontrolledChangelistState.cpp b/Engine/Source/Developer/UncontrolledChangelists/Private/UncontrolledChangelistState.cpp
--- a/Engine/Source/Developer/UncontrolledChangelists/Private/UncontrolledChangelistState.cpp
+++ b/Engine/Source/Developer/UncontrolledChangelists/Private/UncontrolledChangelistState.cpp
@@ -126,11 +126,11 @@ bool FUncontrolledChangelistState::AddFiles(const TArray<FString>& InFilenames,
 	{
 		for (FSourceControlStateRef FileState : FileStates)
 		{
-			const bool bIsSourceControlled = FileState->IsSourceControlled();
+			const bool bIsSourceControlled = (!FileState->IsUnknown()) && FileState->IsSourceControlled();
 			const bool bFileExists = IFileManager::Get().FileExists(*FileState->GetFilename());
 
 			const bool bIsUncontrolled = (!bIsSourceControlled) && bFileExists;
-			const bool bIsDeleted = FileState->IsSourceControlled() && (!bFileExists);
+			const bool bIsDeleted = bIsSourceControlled && (!bFileExists);
 
 			const bool bIsCheckoutCompliant = (!bCheckCheckout) || (!FileState->IsCheckedOut());
 			const bool bIsStatusCompliant = (!bCheckStatus) || FileState->IsModified() || bIsUncontrolled || bIsDeleted;
diff --git a/Engine/Source/Editor/SourceControlWindows/Private/SSourceControlChangelists.cpp b/Engine/Source/Editor/SourceControlWindows/Private/SSourceControlChangelists.cpp
--- a/Engine/Source/Editor/SourceControlWindows/Private/SSourceControlChangelists.cpp
+++ b/Engine/Source/Editor/SourceControlWindows/Private/SSourceControlChangelists.cpp
@@ -1235,7 +1235,12 @@ void SSourceControlChangelistsWidget::OnLocateFile()
 	{
 		if (SelectedItem->GetTreeItemType() == IChangelistTreeItem::File)
 		{
-			AssetsToSync.Append(StaticCastSharedPtr<FFileTreeItem>(SelectedItem)->GetAssetData());
+			const FAssetDataArrayPtr& Assets = StaticCastSharedPtr<FFileTreeItem>(SelectedItem)->GetAssetData();
+
+			if (Assets.IsValid())
+			{
+				AssetsToSync.Append(*Assets);
+			}
 		}
 	}
 
@@ -1252,7 +1257,14 @@ bool SSourceControlChangelistsWidget::CanLocateFile()
 
 	auto HasAssetData = [](const FChangelistTreeItemPtr& SelectedItem)
 	{
-		return (SelectedItem->GetTreeItemType() == IChangelistTreeItem::File) && (StaticCastSharedPtr<FFileTreeItem>(SelectedItem)->GetAssetData().Num() > 0);
+		if (SelectedItem->GetTreeItemType() != IChangelistTreeItem::File)
+		{
+			return false;
+		}
+
+		const FAssetDataArrayPtr& Assets = StaticCastSharedPtr<FFileTreeItem>(SelectedItem)->GetAssetData();
+
+		return (Assets.IsValid() && Assets->Num() > 0);
 	};
 
 	// Checks if at least one selected item has asset data (ie: accessible from ContentBrowser)
diff --git a/Engine/Source/Editor/SourceControlWindows/Private/SSourceControlCommon.h b/Engine/Source/Editor/SourceControlWindows/Private/SSourceControlCommon.h
--- a/Engine/Source/Editor/SourceControlWindows/Private/SSourceControlCommon.h
+++ b/Engine/Source/Editor/SourceControlWindows/Private/SSourceControlCommon.h
@@ -159,9 +159,9 @@ struct FFileTreeItem : public IChangelistTreeItem
 	/** true if the item is source controlled and not marked for add nor for delete */
 	bool CanDiff() const { return FileState->IsSourceControlled() && !FileState->IsAdded() && !FileState->IsDeleted(); }
 
-	const TArray<FAssetData>& GetAssetData() const
+	const FAssetDataArrayPtr& GetAssetData() const
 	{
-		return *Assets;
+		return Assets;
 	}
 
 	bool IsShelved() const { return GetTreeItemType() == IChangelistTreeItem::ShelvedFile; }
