SourceControl: Propagate the cancel result of the FCheckIn operation through FCheckinResultInfo. Callers that just want to know if the operation succeeded can still rely on the return value of the function - as before.

diff --git a/Engine/Source/Editor/SourceControlWindows/Private/SourceControlWindows.cpp b/Engine/Source/Editor/SourceControlWindows/Private/SourceControlWindows.cpp
--- a/Engine/Source/Editor/SourceControlWindows/Private/SourceControlWindows.cpp
+++ b/Engine/Source/Editor/SourceControlWindows/Private/SourceControlWindows.cpp
@@ -258,8 +258,9 @@ bool FSourceControlWindows::PromptForCheckin(FCheckinResultInfo& OutResultInfo,
 	// Mark files for add as needed
 	
 	bool bSuccess = true;  // Overall success
-    bool bAddSuccess = true;
-    bool bCheckinSuccess = true;
+	bool bAddSuccess = true;
+	bool bCheckinSuccess = true;
+	bool bCheckinCancelled = false;
 
 	TArray<FString> CombinedFileList = Description.FilesForAdd;
 	CombinedFileList.Append(Description.FilesForSubmit);
@@ -335,7 +336,10 @@ bool FSourceControlWindows::PromptForCheckin(FCheckinResultInfo& OutResultInfo,
 	TSharedRef<FCheckIn, ESPMode::ThreadSafe> CheckInOperation = ISourceControlOperation::Create<FCheckIn>();
 	CheckInOperation->SetDescription(Description.Description);
 
-	bCheckinSuccess = SourceControlProvider.Execute(CheckInOperation, CombinedFileList) == ECommandResult::Succeeded;
+	ECommandResult::Type CheckInResult = SourceControlProvider.Execute(CheckInOperation, CombinedFileList);
+	bCheckinSuccess = CheckInResult == ECommandResult::Succeeded;
+	bCheckinCancelled = CheckInResult == ECommandResult::Cancelled;
+
 	bSuccess &= bCheckinSuccess;
 
 	if (bCheckinSuccess)
@@ -353,6 +357,18 @@ bool FSourceControlWindows::PromptForCheckin(FCheckinResultInfo& OutResultInfo,
 		OutResultInfo.Description    = CheckInOperation->GetSuccessMessage();
 		OutResultInfo.FilesSubmitted = Description.FilesForSubmit;
 	}
+
+	//~~~~~~~~~~~~~~~~~~~~~~~~~~~~~~~~~~~~~~~~~~
+	// Abort if cancelled
+	if (bCheckinCancelled)
+	{
+		FText Message(LOCTEXT("CheckinCancelled", "File check in cancelled."));
+
+		OutResultInfo.Result = ECommandResult::Cancelled;
+		OutResultInfo.Description = Message;
+
+		return false;
+	}
 	
 	//~~~~~~~~~~~~~~~~~~~~~~~~~~~~~~~~~~~~~~~~~~
 	// Exit if errors
