SourceControl: CL #19752142 propagated the 'cancel' result, but the 'success' result was never set before either.

diff --git a/Engine/Source/Editor/SourceControlWindows/Private/SourceControlWindows.cpp b/Engine/Source/Editor/SourceControlWindows/Private/SourceControlWindows.cpp
--- a/Engine/Source/Editor/SourceControlWindows/Private/SourceControlWindows.cpp
+++ b/Engine/Source/Editor/SourceControlWindows/Private/SourceControlWindows.cpp
@@ -354,6 +354,7 @@ bool FSourceControlWindows::PromptForCheckin(FCheckinResultInfo& OutResultInfo,
 		// also add to the log
 		FMessageLog("SourceControl").Info(CheckInOperation->GetSuccessMessage());
 
+		OutResultInfo.Result         = ECommandResult::Succeeded;
 		OutResultInfo.Description    = CheckInOperation->GetSuccessMessage();
 		OutResultInfo.FilesSubmitted = Description.FilesForSubmit;
 	}
@@ -364,7 +365,7 @@ bool FSourceControlWindows::PromptForCheckin(FCheckinResultInfo& OutResultInfo,
 	{
 		FText Message(LOCTEXT("CheckinCancelled", "File check in cancelled."));
 
-		OutResultInfo.Result = ECommandResult::Cancelled;
+		OutResultInfo.Result      = ECommandResult::Cancelled;
 		OutResultInfo.Description = Message;
 
 		return false;
