Fixed the Source Control SavedChangeListDescription being errased on Construct (Contributed by [REDACTED_PERSON])

#9112: [Source Control] Fix the SavedChangeListDescription being errased on Construct (Contributed by [REDACTED_PERSON])

diff --git a/Engine/Source/Editor/SourceControlWindows/Private/SSourceControlSubmit.cpp b/Engine/Source/Editor/SourceControlWindows/Private/SSourceControlSubmit.cpp
--- a/Engine/Source/Editor/SourceControlWindows/Private/SSourceControlSubmit.cpp
+++ b/Engine/Source/Editor/SourceControlWindows/Private/SSourceControlSubmit.cpp
@@ -67,6 +67,7 @@ FText SSourceControlSubmitWidget::SavedChangeListDescription;
 
 SSourceControlSubmitWidget::~SSourceControlSubmitWidget()
 {
+	// If the user cancel the submit, save the changelist. (On successful submit, ChangeListDescriptionTextCtrl is cleared).
 	SavedChangeListDescription = ChangeListDescriptionTextCtrl->GetText();
 }
 
@@ -75,7 +76,11 @@ void SSourceControlSubmitWidget::Construct(const FArguments& InArgs)
 	ParentFrame = InArgs._ParentWindow.Get();
 	SortByColumn = SSourceControlSubmitWidgetDefs::ColumnID_AssetLabel;
 	SortMode = EColumnSortMode::Ascending;
-	SavedChangeListDescription = InArgs._Description.Get();
+	if (!InArgs._Description.Get().IsEmpty())
+	{
+		// This is meant to restore the description entered by the user if he canceled a previous submit operation (or the submit failed somehow).
+		SavedChangeListDescription = InArgs._Description.Get();
+	}
 	bAllowSubmit = InArgs._AllowSubmit.Get();
 
 	const bool bDescriptionIsReadOnly = !InArgs._AllowDescriptionChange.Get();
