SourceControl: Only allow the changelists tab to be shown if the SourceControlProvider has indicated it's actually supporting that.

diff --git a/Engine/Source/Editor/SourceControlWindows/Private/SourceControlWindowsModule.cpp b/Engine/Source/Editor/SourceControlWindows/Private/SourceControlWindowsModule.cpp
--- a/Engine/Source/Editor/SourceControlWindows/Private/SourceControlWindowsModule.cpp
+++ b/Engine/Source/Editor/SourceControlWindows/Private/SourceControlWindowsModule.cpp
@@ -116,7 +116,7 @@
 {
 	ISourceControlModule& SourceControlModule = ISourceControlModule::Get();
 
-	return (SourceControlModule.IsEnabled() && SourceControlModule.GetProvider().IsAvailable()) || FUncontrolledChangelistsModule::Get().IsEnabled();
+	return (SourceControlModule.IsEnabled() && SourceControlModule.GetProvider().IsAvailable() && SourceControlModule.GetProvider().UsesChangelists()) || FUncontrolledChangelistsModule::Get().IsEnabled();
 }
 
 #undef LOCTEXT_NAMESPACE
