Changed the source control changelist widget tab spawner to a nomad spanwer to allow restoring the widget dock position on next Editor restart.

diff --git a/Engine/Source/Editor/SourceControlWindows/Private/SourceControlWindowsModule.cpp b/Engine/Source/Editor/SourceControlWindows/Private/SourceControlWindowsModule.cpp
--- a/Engine/Source/Editor/SourceControlWindows/Private/SourceControlWindowsModule.cpp
+++ b/Engine/Source/Editor/SourceControlWindows/Private/SourceControlWindowsModule.cpp
@@ -59,7 +59,7 @@
 	const FSlateIcon SourceControlIcon(FAppStyle::GetAppStyleSetName(), "SourceControl.ChangelistsTab");
 
 	// Register the changelist tab spawner
-	FGlobalTabmanager::Get()->RegisterTabSpawner(SourceControlChangelistsTabName, FOnSpawnTab::CreateRaw(this, &FSourceControlWindowsModule::CreateChangelistsTab))
+	FGlobalTabmanager::Get()->RegisterNomadTabSpawner(SourceControlChangelistsTabName, FOnSpawnTab::CreateRaw(this, &FSourceControlWindowsModule::CreateChangelistsTab))
 		.SetDisplayName(LOCTEXT("ChangelistsTabTitle", "View Changelists"))
 		.SetTooltipText(LOCTEXT("ChangelistsTabTooltip", "Opens a dialog displaying current changelists."))
 		.SetIcon(SourceControlIcon);
