Two additional "View Changelists" appeared in menus "Window" and one additional in "Tool". Hid the NomadTabSpawner.

diff --git a/Engine/Source/Editor/SourceControlWindows/Private/SourceControlWindowsModule.cpp b/Engine/Source/Editor/SourceControlWindows/Private/SourceControlWindowsModule.cpp
--- a/Engine/Source/Editor/SourceControlWindows/Private/SourceControlWindowsModule.cpp
+++ b/Engine/Source/Editor/SourceControlWindows/Private/SourceControlWindowsModule.cpp
@@ -62,7 +62,8 @@ void FSourceControlWindowsModule::StartupModule()
 	FGlobalTabmanager::Get()->RegisterNomadTabSpawner(SourceControlChangelistsTabName, FOnSpawnTab::CreateRaw(this, &FSourceControlWindowsModule::CreateChangelistsTab))
 		.SetDisplayName(LOCTEXT("ChangelistsTabTitle", "View Changelists"))
 		.SetTooltipText(LOCTEXT("ChangelistsTabTooltip", "Opens a dialog displaying current changelists."))
-		.SetIcon(SourceControlIcon);
+		.SetIcon(SourceControlIcon)
+		.SetMenuType(ETabSpawnerMenuType::Hidden);	
 
 #if WITH_RELOAD
 	// This code attempts to relaunch the GameplayCueEditor tab when you reload this module
@@ -77,7 +78,7 @@ void FSourceControlWindowsModule::ShutdownModule()
 {
 	if (FSlateApplication::IsInitialized())
 	{
-		FGlobalTabmanager::Get()->UnregisterTabSpawner(SourceControlChangelistsTabName);
+		FGlobalTabmanager::Get()->UnregisterNomadTabSpawner(SourceControlChangelistsTabName);
 
 		if (ChangelistsTab.IsValid())
 		{
