We now call the virtualization process directly when submitting content via the editor rather than invoking a delegate.

#rnx

- Now that the IVirtualizationSystem interface supports both virtualization and rehydration of package files, it seems pointless to keep a delegate that occurs on pre-submit, rather than just calling the system we really want to call directly.
- Deprecating the public methods accessing the delegate from ISourceControlModule should prevent licensees from trying to use it. The delegate was added in 5.0 so even though it is unlikely to have been used by anyone we need to go through the deprecation steps to remove it.
-- This leaves the calling code a little more complicated than I'd like, but cannot be avoided.

diff --git a/Engine/Source/Developer/SourceControl/Public/ISourceControlModule.h b/Engine/Source/Developer/SourceControl/Public/ISourceControlModule.h
--- a/Engine/Source/Developer/SourceControl/Public/ISourceControlModule.h
+++ b/Engine/Source/Developer/SourceControl/Public/ISourceControlModule.h
@@ -98,7 +98,7 @@ public:
 	virtual void QueueStatusUpdate(const TArray<UPackage*>& InPackages) = 0;
 
 	/**
-	 * Queues a file to have its source control status updated in the background.
+	 * Queues a file to have its s7ource control status updated in the background.
 	 * @param	InFilenames	The files to queue.
 	 */
 	virtual void QueueStatusUpdate(const TArray<FString>& InFilenames) = 0;
@@ -204,16 +204,19 @@ public:
 	/** 
 	 * Register a delegate that is invokes right before files are submitted to source control. @see FSourceControlPreSubmitFinalizeDelegate
 	 */
+	UE_DEPRECATED(5.1, "ISourceControlModule::RegisterPreSubmitFinalize is deprecated, the functionality will be removed")
 	virtual FDelegateHandle RegisterPreSubmitFinalize(const FSourceControlPreSubmitFinalizeDelegate::FDelegate& Delegate) = 0;
 
 	/**
 	 * Unregister a previously registered delegate. @see FSourceControlPreSubmitFinalizeDelegate
 	 */
+	UE_DEPRECATED(5.1, "ISourceControlModule::UnregisterPreSubmitFinalize is deprecated, the functionality will be removed")
 	virtual void UnregisterPreSubmitFinalize(FDelegateHandle Handle) = 0;
 
 	/** 
 	 * Returns access to the delegate so that it can be broadcast as needed. @see FSourceControlPreSubmitFinalizeDelegate
 	 */
+	UE_DEPRECATED(5.1, "ISourceControlModule::GetOnPreSubmitFinalize is deprecated, the functionality will be removed")
 	virtual const FSourceControlPreSubmitFinalizeDelegate& GetOnPreSubmitFinalize() const = 0;
 
 	/**
diff --git a/Engine/Source/Developer/Virtualization/Private/Virtualization.cpp b/Engine/Source/Developer/Virtualization/Private/Virtualization.cpp
--- a/Engine/Source/Developer/Virtualization/Private/Virtualization.cpp
+++ b/Engine/Source/Developer/Virtualization/Private/Virtualization.cpp
@@ -3,7 +3,6 @@
 #include "Modules/ModuleManager.h"
 
 #include "Features/IModularFeatures.h"
-#include "ISourceControlModule.h"
 #include "MessageLogModule.h"
 #include "Misc/DelayedAutoRegister.h"
 #include "PackageSubmissionChecks.h"
@@ -22,13 +21,6 @@ public:
 	{
 		IModularFeatures::Get().RegisterModularFeature(FName("VirtualizationSourceControlUtilities"), &SourceControlutility);
 
-		// Delay this until after the source control module has loaded
-		FDelayedAutoRegisterHelper(EDelayedRegisterRunPhase::EarliestPossiblePluginsLoaded, [this]() 
-			{
-				PackageSubmissionHandle = ISourceControlModule::Get().RegisterPreSubmitFinalize(
-					FSourceControlPreSubmitFinalizeDelegate::FDelegate::CreateStatic(&VirtualizePackages));
-			});
-
 		FMessageLogModule& MessageLogModule = FModuleManager::LoadModuleChecked<FMessageLogModule>("MessageLog");
 		MessageLogModule.RegisterLogListing("LogVirtualization", LOCTEXT("AssetVirtualizationLogLabel", "Asset Virtualization"));
 	}
@@ -40,23 +32,12 @@ public:
 			FMessageLogModule& MessageLogModule = FModuleManager::GetModuleChecked<FMessageLogModule>("MessageLog");
 			MessageLogModule.UnregisterLogListing("LogVirtualization&");
 		}
-	
-		// The SourceControl module might be destroyed before this one, depending on shutdown order so we need to check if 
-		// it is loaded before unregistering.
-		if (ISourceControlModule* SourceControlModule = FModuleManager::GetModulePtr<ISourceControlModule>(FName("SourceControl")))
-		{
-			SourceControlModule->UnregisterPreSubmitFinalize(PackageSubmissionHandle);
-		}
-
-		PackageSubmissionHandle.Reset();
 		
 		IModularFeatures::Get().UnregisterModularFeature(FName("VirtualizationSourceControlUtilities"), &SourceControlutility);
 	}
 
 private:
 	Experimental::FVirtualizationSourceControlUtilities SourceControlutility;
-
-	FDelegateHandle PackageSubmissionHandle;
 };
 
 } // namespace UE::Virtualization
diff --git a/Engine/Source/Editor/SourceControlWindows/Private/SSourceControlChangelists.cpp b/Engine/Source/Editor/SourceControlWindows/Private/SSourceControlChangelists.cpp
--- a/Engine/Source/Editor/SourceControlWindows/Private/SSourceControlChangelists.cpp
+++ b/Engine/Source/Editor/SourceControlWindows/Private/SSourceControlChangelists.cpp
@@ -36,7 +36,7 @@
 #include "SSourceControlSubmit.h"
 #include "Framework/Application/SlateApplication.h"
 #include "Framework/Notifications/NotificationManager.h"
-
+#include "Virtualization/VirtualizationSystem.h"
 
 #define LOCTEXT_NAMESPACE "SourceControlChangelist"
 
@@ -936,41 +936,51 @@ static bool GetOnPresubmitResult(FSourceControlChangelistStatePtr Changelist, FC
 
 	TArray<FText> PayloadErrors;
 	TArray<FText> DescriptionTags;
+
+PRAGMA_DISABLE_DEPRECATION_WARNINGS
 	ISourceControlModule::Get().GetOnPreSubmitFinalize().Broadcast(LocalFilepathList, DescriptionTags, PayloadErrors);
+PRAGMA_ENABLE_DEPRECATION_WARNINGS
 
-	if (!PayloadErrors.IsEmpty())
+	UE::Virtualization::IVirtualizationSystem& System = UE::Virtualization::IVirtualizationSystem::Get();
+	if (System.IsEnabled())
 	{
-		for (const FText& Error : PayloadErrors)
+		if (System.TryVirtualizePackages(LocalFilepathList, DescriptionTags, PayloadErrors))
 		{
-			FMessageLog("SourceControl").Error(Error);
-		}
+			if (!DescriptionTags.IsEmpty())
+			{
+				FTextBuilder NewDescription;
+				NewDescription.AppendLine(Description.Description);
+				NewDescription.AppendLine(); // Add a gap between the user input and any automated description tags we need to add
 
-		// Setup the notification for operation feedback
-		FText FailureMsg(LOCTEXT("SCC_Virtualization_Failed", "Virtualized payloads failed to submit."));
-		FNotificationInfo Info(FailureMsg);
+				for (const FText& Line : DescriptionTags)
+				{
+					NewDescription.AppendLine(Line);
+				}
 
-		Info.Text = LOCTEXT("SCC_Checkin_Failed", "Failed to check in files!");
-		Info.ExpireDuration = 8.0f;
-		Info.HyperlinkText = LOCTEXT("SCC_Checkin_ShowLog", "Show Message Log");
-		Info.Hyperlink = FSimpleDelegate::CreateLambda([]() { FMessageLog("SourceControl").Open(EMessageSeverity::Error, true); });
+				Description.Description = NewDescription.ToText();
+			}
+		}
+		else
+		{
+			for (const FText& Error : PayloadErrors)
+			{
+				FMessageLog("SourceControl").Error(Error);
+			}
 
-		TSharedPtr<SNotificationItem> Notification = FSlateNotificationManager::Get().AddNotification(Info);
-		Notification->SetCompletionState(SNotificationItem::CS_Fail);
+			// Setup the notification for operation feedback
+			FText FailureMsg(LOCTEXT("SCC_Virtualization_Failed", "Virtualized payloads failed to submit."));
+			FNotificationInfo Info(FailureMsg);
 
-		return false;
-	}
-	else if (!DescriptionTags.IsEmpty())
-	{
-		FTextBuilder NewDescription;
-		NewDescription.AppendLine(Description.Description);
-		NewDescription.AppendLine(); // Add a gap between the user input and any automated description tags we need to add
+			Info.Text = LOCTEXT("SCC_Checkin_Failed", "Failed to check in files!");
+			Info.ExpireDuration = 8.0f;
+			Info.HyperlinkText = LOCTEXT("SCC_Checkin_ShowLog", "Show Message Log");
+			Info.Hyperlink = FSimpleDelegate::CreateLambda([]() { FMessageLog("SourceControl").Open(EMessageSeverity::Error, true); });
 
-		for (const FText& Line : DescriptionTags)
-		{
-			NewDescription.AppendLine(Line);
-		}
+			TSharedPtr<SNotificationItem> Notification = FSlateNotificationManager::Get().AddNotification(Info);
+			Notification->SetCompletionState(SNotificationItem::CS_Fail);
 
-		Description.Description = NewDescription.ToText();
+			return false;
+		}
 	}
 
 	return true;
diff --git a/Engine/Source/Editor/SourceControlWindows/Private/SourceControlWindows.cpp b/Engine/Source/Editor/SourceControlWindows/Private/SourceControlWindows.cpp
--- a/Engine/Source/Editor/SourceControlWindows/Private/SourceControlWindows.cpp
+++ b/Engine/Source/Editor/SourceControlWindows/Private/SourceControlWindows.cpp
@@ -13,6 +13,7 @@
 #include "Logging/TokenizedMessage.h"
 #include "Misc/MessageDialog.h"
 #include "Widgets/Notifications/SNotificationList.h"
+#include "Virtualization/VirtualizationSystem.h"
 
 
 #if SOURCE_CONTROL_WITH_SLATE
@@ -295,41 +296,54 @@ bool FSourceControlWindows::PromptForCheckin(FCheckinResultInfo& OutResultInfo,
 
 	//~~~~~~~~~~~~~~~~~~~~~~~~~~~~~~~~~~~~~~~~~~
 	// Allow any finalization pre-submit processes to run
+	
+	// TODO: Once this is removed and not deprecated move the following arrays inside of the System.IsEnabled() scope
 	TArray<FText> PayloadErrors;
 	TArray<FText> DescriptionTags;
+
+PRAGMA_DISABLE_DEPRECATION_WARNINGS
 	ISourceControlModule::Get().GetOnPreSubmitFinalize().Broadcast(CombinedFileList, DescriptionTags, PayloadErrors);
+PRAGMA_ENABLE_DEPRECATION_WARNINGS
 
-	if (!PayloadErrors.IsEmpty())
+	UE::Virtualization::IVirtualizationSystem& System = UE::Virtualization::IVirtualizationSystem::Get();
+	if (System.IsEnabled())
 	{
-		for (const FText& Error : PayloadErrors)
+		if (System.TryVirtualizePackages(CombinedFileList, DescriptionTags, PayloadErrors))
 		{
-			FMessageLog("SourceControl").Error(Error);
+			if (!DescriptionTags.IsEmpty())
+			{
+				FTextBuilder NewDescription;
+				NewDescription.AppendLine(Description.Description);
+
+				for (const FText& Line : DescriptionTags)
+				{
+					NewDescription.AppendLine(Line);
+				}
+
+				Description.Description = NewDescription.ToText();
+			}
 		}
-			
-		// TODO: We rely on the calling code to issue a noticeable error message to the user but this is 
-		// currently handled very inconsistently in the engine and the return values often ignored.
-		// So print the error to the log file as well for now but we should try and unify how the errors
-		// are reported.
-		FText FailureMsg(LOCTEXT("SCC_Virtualization_Failed", "Virtualized payloads failed to submit."));
-		FMessageLog("SourceControl").Notify(FailureMsg);
+		else
+		{
+			for (const FText& Error : PayloadErrors)
+			{
+				FMessageLog("SourceControl").Error(Error);
+			}
 
-		OutResultInfo.Result = ECommandResult::Failed;
-		OutResultInfo.Description = FailureMsg;
+			// TODO: We rely on the calling code to issue a noticeable error message to the user but this is 
+			// currently handled very inconsistently in the engine and the return values often ignored.
+			// So print the error to the log file as well for now but we should try and unify how the errors
+			// are reported.
+			FText FailureMsg(LOCTEXT("SCC_Virtualization_Failed", "Virtualized payloads failed to submit."));
+			FMessageLog("SourceControl").Notify(FailureMsg);
 
-		return false;
-	}
-	else if (!DescriptionTags.IsEmpty())
-	{
-		FTextBuilder NewDescription;
-		NewDescription.AppendLine(Description.Description);
+			OutResultInfo.Result = ECommandResult::Failed;
+			OutResultInfo.Description = FailureMsg;
 
-		for (const FText& Line : DescriptionTags)
-		{
-			NewDescription.AppendLine(Line);
+			return false;
 		}
-
-		Description.Description = NewDescription.ToText();
 	}
+
 	//~~~~~~~~~~~~~~~~~~~~~~~~~~~~~~~~~~~~~~~~~~
 
 	// Check in files
