Changelist validations: split changelist validation warnings and errors.

#rnx

diff --git a/Engine/Source/Editor/SourceControlWindows/Private/SSourceControlChangelists.cpp b/Engine/Source/Editor/SourceControlWindows/Private/SSourceControlChangelists.cpp
--- a/Engine/Source/Editor/SourceControlWindows/Private/SSourceControlChangelists.cpp
+++ b/Engine/Source/Editor/SourceControlWindows/Private/SSourceControlChangelists.cpp
@@ -839,7 +839,7 @@ void SSourceControlChangelistsWidget::OnDeleteShelvedFiles()
 	SourceControlProvider.Execute(DeleteShelvedOperation, GetChangelistFromSelection(), GetSelectedShelvedFiles());
 }
 
-static bool GetChangelistValidationResult(FSourceControlChangelistPtr InChangelist, FString& OutValidationText)
+static bool GetChangelistValidationResult(FSourceControlChangelistPtr InChangelist, FString& OutValidationTitleText, FString& OutValidationWarningsText, FString& OutValidationErrorsText)
 {
 	FSourceControlPreSubmitDataValidationDelegate ValidationDelegate = ISourceControlModule::Get().GetRegisteredPreSubmitDataValidation();
 
@@ -855,45 +855,45 @@ static bool GetChangelistValidationResult(FSourceControlChangelistPtr InChangeli
 
 		if (ValidationResult == EDataValidationResult::Invalid || ValidationErrors.Num() > 0)
 		{
-			OutValidationText = LOCTEXT("SourceControl.Submit.ChangelistValidationError", "Changelist validation failed!").ToString();
+			OutValidationTitleText = LOCTEXT("SourceControl.Submit.ChangelistValidationError", "Changelist validation failed!").ToString();
 			bValidationResult = false;
 			MessageSeverity = EMessageSeverity::Error;
 		}
 		else if (ValidationResult == EDataValidationResult::NotValidated || ValidationWarnings.Num() > 0)
 		{
-			OutValidationText = LOCTEXT("SourceControl.Submit.ChangelistValidationWarning", "Changelist validation has warnings!").ToString();
+			OutValidationTitleText = LOCTEXT("SourceControl.Submit.ChangelistValidationWarning", "Changelist validation has warnings!").ToString();
 			MessageSeverity = EMessageSeverity::Warning;
 		}
 		else
 		{
-			OutValidationText = LOCTEXT("SourceControl.Submit.ChangelistValidationSuccess", "Changelist validation successful!").ToString();
+			OutValidationTitleText = LOCTEXT("SourceControl.Submit.ChangelistValidationSuccess", "Changelist validation successful!").ToString();
 		}
 
-		int32 NumLinesDisplayed = 0;
 		FMessageLog SourceControlLog("SourceControl");
 		
-		SourceControlLog.Message(MessageSeverity, FText::FromString(*OutValidationText));
+		SourceControlLog.Message(MessageSeverity, FText::FromString(*OutValidationTitleText));
 
-		auto AppendInfo = [&OutValidationText, &NumLinesDisplayed](const TArray<FText>& Info, const FString& InfoType)
+		auto AppendInfo = [](const TArray<FText>& Info, const FString& InfoType, FString& OutText)
 		{
 			const int32 MaxNumLinesDisplayed = 5;
+			int32 NumLinesDisplayed = 0;
 
 			if (Info.Num() > 0)
 			{
-				OutValidationText += LINE_TERMINATOR;
-				OutValidationText += FString::Printf(TEXT("Encountered %d %s:"), Info.Num(), *InfoType);
+				OutText += LINE_TERMINATOR;
+				OutText += FString::Printf(TEXT("Encountered %d %s:"), Info.Num(), *InfoType);
 
 				for (const FText& Line : Info)
 				{
 					if (NumLinesDisplayed >= MaxNumLinesDisplayed)
 					{
-						OutValidationText += LINE_TERMINATOR;
-						OutValidationText += FString::Printf(TEXT("See log for complete list of %s"), *InfoType);
+						OutText += LINE_TERMINATOR;
+						OutText += FString::Printf(TEXT("See log for complete list of %s"), *InfoType);
 						break;
 					}
 
-					OutValidationText += LINE_TERMINATOR;
-					OutValidationText += Line.ToString();
+					OutText += LINE_TERMINATOR;
+					OutText += Line.ToString();
 
 					++NumLinesDisplayed;
 				}
@@ -913,8 +913,8 @@ static bool GetChangelistValidationResult(FSourceControlChangelistPtr InChangeli
 			}
 		};
 
-		AppendInfo(ValidationErrors, TEXT("errors"));
-		AppendInfo(ValidationWarnings, TEXT("warnings"));
+		AppendInfo(ValidationErrors, TEXT("errors"), OutValidationErrorsText);
+		AppendInfo(ValidationWarnings, TEXT("warnings"), OutValidationWarningsText);
 
 		LogInfo(ValidationErrors, TEXT("errors"), EMessageSeverity::Error);
 		LogInfo(ValidationWarnings, TEXT("warnings"), EMessageSeverity::Warning);
@@ -962,14 +962,15 @@ void SSourceControlChangelistsWidget::OnSubmitChangelist()
 		return;
 	}
 
-	FString ChangelistValidationText;
-	bool bValidationResult = GetChangelistValidationResult(ChangelistState->GetChangelist(), ChangelistValidationText);
+	FString ChangelistValidationTitle;
+	FString ChangelistValidationWarningsText;
+	FString ChangelistValidationErrorsText;
+	bool bValidationResult = GetChangelistValidationResult(ChangelistState->GetChangelist(), ChangelistValidationTitle, ChangelistValidationWarningsText, ChangelistValidationErrorsText);
 
 	// Build list of states for the dialog
 	const FText OriginalChangelistDescription = ChangelistState->GetDescriptionText();
 	const bool bAskForChangelistDescription = (OriginalChangelistDescription.IsEmptyOrWhitespace());
 	FText ChangelistDescriptionToSubmit = UpdateChangelistDescriptionToSubmitIfNeeded(bValidationResult, OriginalChangelistDescription);
-	FName ChangelistValidationIconName = bValidationResult ? TEXT("Icons.SuccessWithColor.Large") : TEXT("Icons.ErrorWithColor.Large");
 
 	TSharedRef<SWindow> NewWindow = SNew(SWindow)
 		.Title(NSLOCTEXT("SourceControl.ConfirmSubmit", "Title", "Confirm changelist submit"))
@@ -983,8 +984,9 @@ void SSourceControlChangelistsWidget::OnSubmitChangelist()
 		.ParentWindow(NewWindow)
 		.Items(ChangelistState->GetFilesStates())
 		.Description(ChangelistDescriptionToSubmit)
-		.ChangeValidationDescription(ChangelistValidationText)
-		.ChangeValidationIcon(ChangelistValidationIconName)
+		.ChangeValidationResult(ChangelistValidationTitle)
+		.ChangeValidationWarnings(ChangelistValidationWarningsText)
+		.ChangeValidationErrors(ChangelistValidationErrorsText)
 		.AllowDescriptionChange(bAskForChangelistDescription)
 		.AllowUncheckFiles(false)
 		.AllowKeepCheckedOut(false)
@@ -1070,8 +1072,10 @@ void SSourceControlChangelistsWidget::OnValidateChangelist()
 		return;
 	}
 
-	FString ChangelistValidationText;
-	bool bValidationResult = GetChangelistValidationResult(ChangelistState->GetChangelist(), ChangelistValidationText);
+	FString ChangelistValidationTitle;
+	FString ChangelistValidationWarningsText;
+	FString ChangelistValidationErrorsText;
+	bool bValidationResult = GetChangelistValidationResult(ChangelistState->GetChangelist(), ChangelistValidationTitle, ChangelistValidationWarningsText, ChangelistValidationErrorsText);
 
 	// Setup the notification for operation feedback
 	FNotificationInfo Info(LOCTEXT("SCC_Validation_Success", "Changelist validated"));
diff --git a/Engine/Source/Editor/SourceControlWindows/Private/SSourceControlSubmit.cpp b/Engine/Source/Editor/SourceControlWindows/Private/SSourceControlSubmit.cpp
--- a/Engine/Source/Editor/SourceControlWindows/Private/SSourceControlSubmit.cpp
+++ b/Engine/Source/Editor/SourceControlWindows/Private/SSourceControlSubmit.cpp
@@ -142,7 +142,7 @@ void SSourceControlSubmitWidget::Construct(const FArguments& InArgs)
 	const bool bDescriptionIsReadOnly = !InArgs._AllowDescriptionChange.Get();
 	const bool bAllowUncheckFiles = InArgs._AllowUncheckFiles.Get();
 	const bool bAllowKeepCheckedOut = InArgs._AllowKeepCheckedOut.Get();
-	const bool bShowChangelistValidation = !InArgs._ChangeValidationDescription.Get().IsEmpty();
+	const bool bShowChangelistValidation = !InArgs._ChangeValidationResult.Get().IsEmpty();
 
 	for (const auto& Item : InArgs._Items.Get())
 	{
@@ -259,29 +259,89 @@ void SSourceControlSubmitWidget::Construct(const FArguments& InArgs)
 
 	if (bShowChangelistValidation)
 	{
-		FString ChangelistResultText = InArgs._ChangeValidationDescription.Get();
-		FName ChangelistIconName = InArgs._ChangeValidationIcon.Get();
+		const FString ChangelistResultText = InArgs._ChangeValidationResult.Get();
+		const FString ChangelistResultWarningsText = InArgs._ChangeValidationWarnings.Get();
+		const FString ChangelistResultErrorsText = InArgs._ChangeValidationErrors.Get();
 
-		Contents->AddSlot()
-		.AutoHeight()
-		.Padding(FMargin(5))
-		[
-			SNew(SHorizontalBox)
-			+SHorizontalBox::Slot()
-			.AutoWidth()
-			.VAlign(VAlign_Center)
-			[
-				SNew(SImage)
-				.Image(FAppStyle::GetBrush(ChangelistIconName))
-			]
-			+SHorizontalBox::Slot()
+		const FName ChangelistSuccessIconName = TEXT("Icons.SuccessWithColor.Large");
+		const FName ChangelistWarningsIconName = TEXT("Icons.WarningWithColor.Large");
+		const FName ChangelistErrorsIconName = TEXT("Icons.ErrorWithColor.Large");
+
+		const bool bIsSuccess = ChangelistResultWarningsText.IsEmpty() && ChangelistResultErrorsText.IsEmpty();
+
+		if (bIsSuccess)
+		{
+			Contents->AddSlot()
+			.AutoHeight()
+			.Padding(FMargin(5))
 			[
-				SNew(SMultiLineEditableTextBox)
-				.Text(FText::FromString(ChangelistResultText))
-				.AutoWrapText(true)
-				.IsReadOnly(true)
-			]
-		];
+				SNew(SHorizontalBox)
+				+SHorizontalBox::Slot()
+				.AutoWidth()
+				.VAlign(VAlign_Center)
+				[
+					SNew(SImage)
+					.Image(FAppStyle::GetBrush(ChangelistSuccessIconName))
+				]
+				+SHorizontalBox::Slot()
+				[
+					SNew(SMultiLineEditableTextBox)
+					.Text(FText::FromString(ChangelistResultText))
+					.AutoWrapText(true)
+					.IsReadOnly(true)
+				]
+			];
+		}
+		else
+		{
+			if (!ChangelistResultErrorsText.IsEmpty())
+			{
+				Contents->AddSlot()
+				.AutoHeight()
+				.Padding(FMargin(5))
+				[
+					SNew(SHorizontalBox)
+					+SHorizontalBox::Slot()
+					.AutoWidth()
+					.VAlign(VAlign_Center)
+					[
+						SNew(SImage)
+						.Image(FAppStyle::GetBrush(ChangelistErrorsIconName))
+					]
+					+SHorizontalBox::Slot()
+					[
+						SNew(SMultiLineEditableTextBox)
+						.Text(FText::FromString(ChangelistResultErrorsText))
+						.AutoWrapText(true)
+						.IsReadOnly(true)
+					]
+				];
+			}
+
+			if (!ChangelistResultWarningsText.IsEmpty())
+			{
+				Contents->AddSlot()
+				.AutoHeight()
+				.Padding(FMargin(5))
+				[
+					SNew(SHorizontalBox)
+					+SHorizontalBox::Slot()
+					.AutoWidth()
+					.VAlign(VAlign_Center)
+					[
+						SNew(SImage)
+						.Image(FAppStyle::GetBrush(ChangelistWarningsIconName))
+					]
+					+SHorizontalBox::Slot()
+					[
+						SNew(SMultiLineEditableTextBox)
+						.Text(FText::FromString(ChangelistResultWarningsText))
+						.AutoWrapText(true)
+						.IsReadOnly(true)
+					]
+				];
+			}
+		}
 	}
 
 	if (bAllowKeepCheckedOut)
diff --git a/Engine/Source/Editor/SourceControlWindows/Private/SSourceControlSubmit.h b/Engine/Source/Editor/SourceControlWindows/Private/SSourceControlSubmit.h
--- a/Engine/Source/Editor/SourceControlWindows/Private/SSourceControlSubmit.h
+++ b/Engine/Source/Editor/SourceControlWindows/Private/SSourceControlSubmit.h
@@ -49,8 +49,9 @@ public:
 		: _ParentWindow()
 		, _Items()
 		, _Description()
-		, _ChangeValidationIcon()
-		, _ChangeValidationDescription()
+		, _ChangeValidationResult()
+		, _ChangeValidationWarnings()
+		, _ChangeValidationErrors()
 		, _AllowDescriptionChange(true)
 		, _AllowUncheckFiles(true)
 		, _AllowKeepCheckedOut(true)
@@ -60,8 +61,9 @@ public:
 		SLATE_ATTRIBUTE(TSharedPtr<SWindow>, ParentWindow)
 		SLATE_ATTRIBUTE(TArray<FSourceControlStateRef>, Items)
 		SLATE_ATTRIBUTE(FText, Description)
-		SLATE_ATTRIBUTE(FName, ChangeValidationIcon)
-		SLATE_ATTRIBUTE(FString, ChangeValidationDescription)
+		SLATE_ATTRIBUTE(FString, ChangeValidationResult)
+		SLATE_ATTRIBUTE(FString, ChangeValidationWarnings)
+		SLATE_ATTRIBUTE(FString, ChangeValidationErrors)
 		SLATE_ATTRIBUTE(bool, AllowDescriptionChange)
 		SLATE_ATTRIBUTE(bool, AllowUncheckFiles)
 		SLATE_ATTRIBUTE(bool, AllowKeepCheckedOut)
