Refactored the Changelist Window UI to work better with large changelists along with bug fixes.

This submit should not regress any existing functionalites of the previous widget.

Along with the user interface changes, notable bug fixes were included
  - If a changelist is not specified for P4 operations 'Mark for Add', 'Mark for Delete and 'Checkout', defaulted to the 'default' P4 changlist. Without a valid changelist, the cache wasn't properly updated.
  - Fixed updating the cache for P4 'Mark for delete' operation not correctly updating the file changelist cached in the file state.

Other changes:
  - P4 changelists are now always returned sorted by changelist number, ascending.
  - Some of the safe operations were running synchronous are not running asynchrnous (create new changelist, delete empty changelist, delete shelved files).

This submit addresses the following Jiras:

diff --git a/Engine/Plugins/Developer/PerforceSourceControl/Source/PerforceSourceControl/Private/PerforceSourceControlChangelistState.cpp b/Engine/Plugins/Developer/PerforceSourceControl/Source/PerforceSourceControl/Private/PerforceSourceControlChangelistState.cpp
--- a/Engine/Plugins/Developer/PerforceSourceControl/Source/PerforceSourceControl/Private/PerforceSourceControlChangelistState.cpp
+++ b/Engine/Plugins/Developer/PerforceSourceControl/Source/PerforceSourceControl/Private/PerforceSourceControlChangelistState.cpp
@@ -12,12 +12,13 @@ FPerforceSourceControlChangelistState::FPerforceSourceControlChangelistState(con
 
 FName FPerforceSourceControlChangelistState::GetIconName() const
 {
-	return FName("SourceControl.Changelist");
+	// Mimic P4V colors, returning the red icon if there are active file(s), the blue if the changelist is empty or all the files are shelved.
+	return Files.Num() > 0 ? FName("SourceControl.Changelist") : FName("SourceControl.ShelvedCHangelist");
 }
 
 FName FPerforceSourceControlChangelistState::GetSmallIconName() const
 {
-	return FName("SourceControl.Changelist");
+	return GetIconName();
 }
 
 FText FPerforceSourceControlChangelistState::GetDisplayText() const
diff --git a/Engine/Plugins/Developer/PerforceSourceControl/Source/PerforceSourceControl/Private/PerforceSourceControlOperations.cpp b/Engine/Plugins/Developer/PerforceSourceControl/Source/PerforceSourceControl/Private/PerforceSourceControlOperations.cpp
--- a/Engine/Plugins/Developer/PerforceSourceControl/Source/PerforceSourceControl/Private/PerforceSourceControlOperations.cpp
+++ b/Engine/Plugins/Developer/PerforceSourceControl/Source/PerforceSourceControl/Private/PerforceSourceControlOperations.cpp
@@ -260,13 +260,16 @@ static void UpdateChangelistState(FPerforceSourceControlProvider& SCCProvider, c
 
 		for (TMap<FString, EPerforceState::Type>::TConstIterator It(InOperationResults); It; ++It)
 		{
-			if ((It.Value() != EPerforceState::CheckedOut) && (It.Value() != EPerforceState::OpenForAdd))
+			if ((It.Value() != EPerforceState::CheckedOut) && (It.Value() != EPerforceState::OpenForAdd) && It.Value() != EPerforceState::MarkedForDelete)
 			{
 				continue;
 			}
 
 			TSharedRef<FPerforceSourceControlState, ESPMode::ThreadSafe> State = SCCProvider.GetStateInternal(It.Key());
 			ChangelistState->Files.Add(State);
+
+			// Keep the changelist stored with cached file state in sync with the actual changelist that owns this file.
+			State->Changelist = InChangelist;
 		}
 	}
 }
diff --git a/Engine/Plugins/Developer/PerforceSourceControl/Source/PerforceSourceControl/Private/PerforceSourceControlOperations.h b/Engine/Plugins/Developer/PerforceSourceControl/Source/PerforceSourceControl/Private/PerforceSourceControlOperations.h
--- a/Engine/Plugins/Developer/PerforceSourceControl/Source/PerforceSourceControl/Private/PerforceSourceControlOperations.h
+++ b/Engine/Plugins/Developer/PerforceSourceControl/Source/PerforceSourceControl/Private/PerforceSourceControlOperations.h
@@ -30,6 +30,7 @@ class FPerforceCheckOutWorker final : public IPerforceSourceControlWorker
 public:
 	FPerforceCheckOutWorker(FPerforceSourceControlProvider& InSourceControlProvider)
 		: IPerforceSourceControlWorker(InSourceControlProvider)
+		, InChangelist(FPerforceSourceControlChangelist::DefaultChangelist) // By default, add checked out files in the default changelist.
 	{}
 	virtual ~FPerforceCheckOutWorker() = default;
 
@@ -91,6 +92,7 @@ class FPerforceMarkForAddWorker final : public IPerforceSourceControlWorker
 public:
 	FPerforceMarkForAddWorker(FPerforceSourceControlProvider& InSourceControlProvider)
 		: IPerforceSourceControlWorker(InSourceControlProvider)
+		, InChangelist(FPerforceSourceControlChangelist::DefaultChangelist) // By default, add new files in the default changelist.
 	{}
 	virtual ~FPerforceMarkForAddWorker() = default;
 
@@ -110,6 +112,7 @@ class FPerforceDeleteWorker final : public IPerforceSourceControlWorker
 public:
 	FPerforceDeleteWorker(FPerforceSourceControlProvider& InSourceControlProvider)
 		: IPerforceSourceControlWorker(InSourceControlProvider)
+		, Changelist(FPerforceSourceControlChangelist::DefaultChangelist) // By default, add deleted files in the default changelist
 	{}
 	virtual ~FPerforceDeleteWorker() = default;
 
diff --git a/Engine/Plugins/Developer/PerforceSourceControl/Source/PerforceSourceControl/Private/PerforceSourceControlProvider.cpp b/Engine/Plugins/Developer/PerforceSourceControl/Source/PerforceSourceControl/Private/PerforceSourceControlProvider.cpp
--- a/Engine/Plugins/Developer/PerforceSourceControl/Source/PerforceSourceControl/Private/PerforceSourceControlProvider.cpp
+++ b/Engine/Plugins/Developer/PerforceSourceControl/Source/PerforceSourceControl/Private/PerforceSourceControlProvider.cpp
@@ -94,7 +94,7 @@ TSharedRef<FPerforceSourceControlState, ESPMode::ThreadSafe> FPerforceSourceCont
 	else
 	{
 		// cache an unknown state for this item
-		TSharedRef<FPerforceSourceControlState, ESPMode::ThreadSafe> NewState = MakeShareable( new FPerforceSourceControlState(Filename) );
+		TSharedRef<FPerforceSourceControlState, ESPMode::ThreadSafe> NewState = MakeShared<FPerforceSourceControlState>(Filename);
 		StateCache.Add(Filename, NewState);
 		return NewState;
 	}
@@ -111,7 +111,7 @@ TSharedRef<FPerforceSourceControlChangelistState, ESPMode::ThreadSafe> FPerforce
 	else
 	{
 		// cache an unknown state for this item
-		TSharedRef<FPerforceSourceControlChangelistState, ESPMode::ThreadSafe> NewState = MakeShareable(new FPerforceSourceControlChangelistState(InChangelist));
+		TSharedRef<FPerforceSourceControlChangelistState, ESPMode::ThreadSafe> NewState = MakeShared<FPerforceSourceControlChangelistState>(InChangelist);
 		ChangelistsStateCache.Add(InChangelist, NewState);
 		return NewState;
 	}
@@ -520,6 +520,8 @@ void FPerforceSourceControlProvider::OutputCommandMessages(const FPerforceSource
 
 void FPerforceSourceControlProvider::Tick()
 {
+	TRACE_CPUPROFILER_EVENT_SCOPE(FPerforceSourceControlProvider::Tick);
+
 	bool bStatesUpdated = false;
 	for (int32 CommandIndex = 0; CommandIndex < CommandQueue.Num(); ++CommandIndex)
 	{
@@ -625,7 +627,7 @@ TArray< TSharedRef<ISourceControlLabel> > FPerforceSourceControlProvider::GetLab
 	return Labels;
 }
 
-TArray<FSourceControlChangelistRef> FPerforceSourceControlProvider::GetChangelists( EStateCacheUsage::Type InStateCacheUsage )
+TArray<FSourceControlChangelistRef> FPerforceSourceControlProvider::GetChangelists(EStateCacheUsage::Type InStateCacheUsage)
 {
 	if (!IsEnabled())
 	{
@@ -641,7 +643,18 @@ TArray<FSourceControlChangelistRef> FPerforceSourceControlProvider::GetChangelis
 	}
 
 	TArray<FSourceControlChangelistRef> Changelists;
-	Algo::Transform(ChangelistsStateCache, Changelists, [](const auto& Pair) { return MakeShared<FPerforceSourceControlChangelist, ESPMode::ThreadSafe>(Pair.Key); });
+	Algo::Transform(ChangelistsStateCache, Changelists, [](const TPair<FPerforceSourceControlChangelist, TSharedRef<FPerforceSourceControlChangelistState, ESPMode::ThreadSafe>>& Pair)
+	{
+		return MakeShared<FPerforceSourceControlChangelist, ESPMode::ThreadSafe>(Pair.Key);
+	});
+
+	// NOTE: Sort in ascending number. If this behavior needs to be configurable, we could have 3-state enum param: 'default, ascending, descending'. For P4, the 'default'
+	//       should be ascending. The changelists are source control agnostics, so sorting a changelist files is probably not a notion that should be leaked in the generic interface.
+	Changelists.Sort([](const FSourceControlChangelistRef& Lhs, const FSourceControlChangelistRef& Rhs)
+	{
+		return static_cast<const FPerforceSourceControlChangelist&>(Lhs.Get()).ToInt() < static_cast<const FPerforceSourceControlChangelist&>(Rhs.Get()).ToInt();
+	});
+
 	return Changelists;
 }
 
@@ -846,9 +859,9 @@ ECommandResult::Type FPerforceSourceControlProvider::ExecuteSynchronousCommand(F
 	}
 
 	// If the command failed, inform the user that they need to try again
-	if ( !InCommand.bCancelled && Result != ECommandResult::Succeeded && !bSuppressResponseMsg )
+	if (!InCommand.bCancelled && Result != ECommandResult::Succeeded && !bSuppressResponseMsg)
 	{
-		FMessageDialog::Open( EAppMsgType::Ok, LOCTEXT("Perforce_ServerUnresponsive", "Perforce server is unresponsive. Please check your connection and try again.") );
+		FMessageDialog::Open(EAppMsgType::Ok, LOCTEXT("Perforce_ServerUnresponsive", "Perforce server is unresponsive. Please check your connection and try again."));
 	}
 
 	// Delete the command now if not marked as auto-delete
diff --git a/Engine/Source/Developer/SourceControl/Public/ISourceControlProvider.h b/Engine/Source/Developer/SourceControl/Public/ISourceControlProvider.h
--- a/Engine/Source/Developer/SourceControl/Public/ISourceControlProvider.h
+++ b/Engine/Source/Developer/SourceControl/Public/ISourceControlProvider.h
@@ -273,7 +273,12 @@ public:
 	 */
 	virtual TArray< TSharedRef<class ISourceControlLabel> > GetLabels( const FString& InMatchingSpec ) const = 0;
 
-	virtual TArray<FSourceControlChangelistRef> GetChangelists( EStateCacheUsage::Type InStateCacheUsage ) = 0;
+	/**
+	 * Returns the list of available changelists if the underlying source control supports the 'changelist' concept.
+	 *
+	 * @param InStateCacheUsage True to retrieve the list from a local cache, false to request it from the server (if any).
+	 */
+	virtual TArray<FSourceControlChangelistRef> GetChangelists(EStateCacheUsage::Type InStateCacheUsage) = 0;
 
 	/**
 	 * Executes the FDownloadFile operation, but unlike the ::Execute method this can be called from a background thread, this
diff --git a/Engine/Source/Editor/SourceControlWindows/Private/SSourceControlChangelists.cpp b/Engine/Source/Editor/SourceControlWindows/Private/SSourceControlChangelists.cpp
--- a/Engine/Source/Editor/SourceControlWindows/Private/SSourceControlChangelists.cpp
+++ b/Engine/Source/Editor/SourceControlWindows/Private/SSourceControlChangelists.cpp
@@ -3,17 +3,18 @@
 #include "SSourceControlChangelists.h"
 
 #include "Styling/AppStyle.h"
-
 #include "Algo/Transform.h"
-
 #include "Logging/MessageLog.h"
 
 #include "Widgets/DeclarativeSyntaxSupport.h"
 #include "Widgets/Images/SImage.h"
 #include "Widgets/Layout/SScrollBorder.h"
 #include "Widgets/Input/SButton.h"
+#include "Widgets/Input/SSearchBox.h"
 #include "Widgets/Layout/SBorder.h"
 #include "Widgets/Layout/SBox.h"
+#include "Widgets/Layout/SSplitter.h"
+#include "Widgets/Layout/SExpandableArea.h"
 #include "Widgets/Notifications/SNotificationList.h"
 
 #include "ISourceControlProvider.h"
@@ -21,7 +22,6 @@
 #include "UncontrolledChangelistsModule.h"
 #include "SourceControlOperations.h"
 #include "ToolMenus.h"
-#include "Widgets/Images/SLayeredImage.h"
 #include "SSourceControlDescription.h"
 #include "SourceControlWindows.h"
 #include "SourceControlHelpers.h"
@@ -32,18 +32,68 @@
 #include "Misc/MessageDialog.h"
 #include "Misc/ScopedSlowTask.h"
 #include "Algo/AnyOf.h"
+#include "HAL/PlatformTime.h"
 
 #include "SSourceControlSubmit.h"
 #include "Framework/Application/SlateApplication.h"
 #include "Framework/Notifications/NotificationManager.h"
+#include "OutputLogModule.h"
 
 #define LOCTEXT_NAMESPACE "SourceControlChangelist"
 
+
 const FText SSourceControlChangelistsWidget::ChangelistValidatedTag = LOCTEXT("ValidationTag", "#changelist validated");
 
-DEFINE_LOG_CATEGORY_STATIC(LogSourceControlChangelist, All, All);
+/** Wraps the execution of a changelist operations with a slow task. */
+void ExecuteChangelistOperationWithSlowTaskWrapper(const FText& Message, const TFunction<void()>& ChangelistTask)
+{
+	FScopedSlowTask Progress(0.f, Message);
+	Progress.MakeDialog();
+	ChangelistTask();
+}
+
+/** Wraps the execution of an uncontrolled changelist operations with a slow task. */
+void ExecuteUncontrolledChangelistOperationWithSlowTaskWrapper(const FText& Message, const TFunction<void()>& UncontrolledChangelistTask)
+{
+	ExecuteChangelistOperationWithSlowTaskWrapper(Message, UncontrolledChangelistTask);
+}
+
+/** Displays toast notification to report the status of task. */
+void DisplaySourceControlOperationNotification(const FText& Message, SNotificationItem::ECompletionState CompletionState)
+{
+	if (Message.IsEmpty())
+	{
+		return;
+	}
+
+	FNotificationInfo NotificationInfo(Message);
+	NotificationInfo.ExpireDuration = 6.0f;
+	//NotificationInfo.Hyperlink = FSimpleDelegate::CreateLambda([]() { FOutputLogModule::Get().FocusOutputLog(); }); // NOTE: This is what we want moving further, but this is only in main branch for the moment.
+	NotificationInfo.Hyperlink = FSimpleDelegate::CreateLambda([]() { FGlobalTabmanager::Get()->TryInvokeTab(FName("OutputLog")); }); // Replace this once with the above once this change reach the main branch.
+	NotificationInfo.HyperlinkText = LOCTEXT("ShowOutputLogHyperlink", "Show Output Log");
+	FSlateNotificationManager::Get().AddNotification(NotificationInfo)->SetCompletionState(CompletionState);
+}
 
-//////////////////////////////
+/** Returns true if a source control provider is enable and support changeslists. */
+bool AreControlledChangelistsEnabled()
+{
+	return ISourceControlModule::Get().IsEnabled() && ISourceControlModule::Get().GetProvider().UsesChangelists();
+};
+
+/** Returns true if Uncontrolled changelists are enabled. */
+bool AreUncontrolledChangelistsEnabled()
+{
+	return FUncontrolledChangelistsModule::Get().IsEnabled();
+};
+
+/** Returns true if they are changelists to display. */
+bool AreChangelistsEnabled()
+{
+	return AreControlledChangelistsEnabled() || AreUncontrolledChangelistsEnabled();
+};
+
+
+/** Implements drag and drop operation. */
 struct FSCCFileDragDropOp : public FDragDropOperation
 {
 	DRAG_DROP_OPERATOR_TYPE(FSCCFileDragDropOp, FDragDropOperation);
@@ -61,11 +111,198 @@ struct FSCCFileDragDropOp : public FDragDropOperation
 	TArray<FSourceControlStateRef> UncontrolledFiles;
 };
 
-//////////////////////////////
-SSourceControlChangelistsWidget::SSourceControlChangelistsWidget()
+
+DECLARE_DELEGATE(FOnSearchBoxExpanded)
+
+/** A button that expands a search box below itself when clicked. */
+class SExpandableSearchButton : public SCompoundWidget
 {
-	bIsRefreshing = false;
-}
+public:
+	SLATE_BEGIN_ARGS(SExpandableSearchButton)
+		: _Style(&FAppStyle::Get().GetWidgetStyle<FSearchBoxStyle>("SearchBox"))
+	{}
+		/** Search box style (used to match the glass icon) */
+		SLATE_STYLE_ARGUMENT(FSearchBoxStyle, Style)
+
+		/** Event fired when the associated search box is made visible */
+		SLATE_EVENT(FOnSearchBoxExpanded, OnSearchBoxExpanded)
+	SLATE_END_ARGS()
+
+	void Construct(const FArguments& InArgs, TSharedRef<SSearchBox> SearchBox)
+	{
+		OnSearchBoxExpanded = InArgs._OnSearchBoxExpanded;
+		SearchStyle = InArgs._Style;
+
+		SearchBox->SetVisibility(TAttribute<EVisibility>::CreateSP(this, &SExpandableSearchButton::GetSearchBoxVisibility));
+		SearchBoxPtr = SearchBox;
+
+		ChildSlot
+		[
+			SNew(SCheckBox)
+			.IsChecked(this, &SExpandableSearchButton::GetToggleButtonState)
+			.OnCheckStateChanged(this, &SExpandableSearchButton::OnToggleButtonStateChanged)
+			.Style(FAppStyle::Get(), "ToggleButtonCheckbox")
+			.Padding(4.0f)
+			.ToolTipText(NSLOCTEXT("ExpandableSearchArea", "ExpandCollapseSearchButton", "Expands or collapses the search text box"))
+			[
+				SNew(SImage)
+				.Image(&SearchStyle->GlassImage)
+				.ColorAndOpacity(FSlateColor::UseForeground())
+			]
+		];
+	}
+
+private:
+	/** Sets whether or not the search area is expanded to expose the search box */
+	void OnToggleButtonStateChanged(ECheckBoxState CheckBoxState)
+	{
+		bIsExpanded = CheckBoxState == ECheckBoxState::Checked;
+
+		if (TSharedPtr<SSearchBox> SearchBox = SearchBoxPtr.Pin())
+		{
+			if (bIsExpanded)
+			{
+				OnSearchBoxExpanded.ExecuteIfBound();
+
+				// Focus the search box when it's shown
+				FSlateApplication::Get().SetUserFocus(FSlateApplication::Get().GetUserIndexForKeyboard(), SearchBox, EFocusCause::SetDirectly);
+			}
+			else
+			{
+				// Clear the search box when it's hidden
+				SearchBox->SetText(FText::GetEmpty());
+			}
+		}
+	}
+
+	ECheckBoxState GetToggleButtonState() const { return bIsExpanded ? ECheckBoxState::Checked : ECheckBoxState::Unchecked; }
+	EVisibility GetSearchBoxVisibility() const { return bIsExpanded ? EVisibility::Visible : EVisibility::Collapsed; }
+
+private:
+	const FSearchBoxStyle* SearchStyle;
+	FOnSearchBoxExpanded OnSearchBoxExpanded;
+	TWeakPtr<SSearchBox> SearchBoxPtr;
+	bool bIsExpanded = false;
+};
+
+/** An expanded area to contain the changelists tree view or then uncontrolled changelists tree view. */
+class SExpandableChangelistArea : public SCompoundWidget
+{
+public:
+	SLATE_BEGIN_ARGS(SExpandableChangelistArea)
+		: _Style(&FAppStyle::Get().GetWidgetStyle<FSearchBoxStyle>("SearchBox"))
+		, _HeaderText()
+		, _ChangelistView()
+		, _OnSearchBoxExpanded()
+		, _OnNewChangelist()
+		, _OnNewChangelistTooltip()
+		, _NewButtonVisibility(EVisibility::Visible)
+		, _SearchButtonVisibility(EVisibility::Visible)
+	{}
+		/** Search box style (used to match the glass icon) */
+		SLATE_STYLE_ARGUMENT(FSearchBoxStyle, Style)
+		/** Text displayed on the expandable area */
+		SLATE_ATTRIBUTE(FText, HeaderText)
+		/** The tree element displayed as body. */
+		SLATE_ARGUMENT(TSharedPtr<SChangelistTree>, ChangelistView)
+		/** Event fired when the associated search box is made visible */
+		SLATE_EVENT(FOnSearchBoxExpanded, OnSearchBoxExpanded)
+		/** Event fired when the 'plus' button is clicked. */
+		SLATE_EVENT(FOnClicked, OnNewChangelist)
+		/** Tooltip displayed over the 'plus' button. */
+		SLATE_ATTRIBUTE(FText, OnNewChangelistTooltip)
+		/** Make the 'plus' button visible or not. */
+		SLATE_ARGUMENT(EVisibility, NewButtonVisibility)
+		/** Make the 'search' button visible or not. */
+		SLATE_ARGUMENT(EVisibility, SearchButtonVisibility)
+	SLATE_END_ARGS()
+
+	void Construct(const FArguments& InArgs)
+	{
+		SearchBox = SNew(SSearchBox);
+
+		ChildSlot
+		[
+			SAssignNew(ExpandableArea, SExpandableArea)
+			.BorderImage(FAppStyle::Get().GetBrush("Brushes.Header"))
+			.BodyBorderImage(FAppStyle::Get().GetBrush("Brushes.Recessed"))
+			.HeaderPadding(FMargin(4.0f, 3.0f))
+			.AllowAnimatedTransition(false)
+			.HeaderContent()
+			[
+				SNew(SHorizontalBox)
+				+SHorizontalBox::Slot()
+				.VAlign(VAlign_Center)
+				[
+					SNew(STextBlock)
+					.Text(InArgs._HeaderText)
+					.TextStyle(FAppStyle::Get(), "ButtonText")
+					.Font(FAppStyle::Get().GetFontStyle("NormalFontBold"))
+				]
+				+SHorizontalBox::Slot()
+				.VAlign(VAlign_Center)
+				.HAlign(HAlign_Right)
+				.AutoWidth()
+				.Padding(4.0f, 0.0f, 0.0f, 0.0f)
+				[
+					SNew(SButton)
+					.ButtonStyle(FAppStyle::Get(), "SimpleButton")
+					.ToolTipText(InArgs._OnNewChangelistTooltip)
+					.OnClicked(InArgs._OnNewChangelist)
+					.ContentPadding(FMargin(1, 0))
+					.Visibility(InArgs._NewButtonVisibility)
+					[
+						SNew(SImage)
+						.Image(FAppStyle::Get().GetBrush("Icons.PlusCircle"))
+						.ColorAndOpacity(FSlateColor::UseForeground())
+					]
+				]
+				+SHorizontalBox::Slot()
+				.VAlign(VAlign_Center)
+				.HAlign(HAlign_Right)
+				.AutoWidth()
+				.Padding(4.0f, 0.0f)
+				[
+					SNew(SBox)
+					.Visibility(InArgs._SearchButtonVisibility)
+					[
+						SNew(SExpandableSearchButton, SearchBox.ToSharedRef())
+					]
+				]
+			]
+			.BodyContent()
+			[
+				SNew(SVerticalBox)
+				+ SVerticalBox::Slot()
+				.AutoHeight()
+				[
+					// Should blend in visually with the header but technically acts like part of the body
+					SNew(SBorder)
+					.BorderImage(FAppStyle::Get().GetBrush("Brushes.Header"))
+					.Padding(FMargin(4.0f, 2.0f))
+					[
+						SearchBox.ToSharedRef()
+					]
+				]
+				+ SVerticalBox::Slot()
+				[
+					SNew(SBorder)
+					.BorderImage(FAppStyle::Get().GetBrush("Brushes.Recessed"))
+					[
+						InArgs._ChangelistView.ToSharedRef()
+					]
+				]
+			]
+		];
+	}
+
+	bool IsExpanded() const { return ExpandableArea->IsExpanded(); }
+
+private:
+	TSharedPtr<SExpandableArea> ExpandableArea;
+	TSharedPtr<SSearchBox> SearchBox;
+};
+
 
 void SSourceControlChangelistsWidget::Construct(const FArguments& InArgs)
 {
@@ -77,12 +314,28 @@ void SSourceControlChangelistsWidget::Construct(const FArguments& InArgs)
 	SourceControlStateChangedDelegateHandle = SCCModule.GetProvider().RegisterSourceControlStateChanged_Handle(FSourceControlStateChanged::FDelegate::CreateSP(this, &SSourceControlChangelistsWidget::OnSourceControlStateChanged));
 	UncontrolledChangelistModule.OnUncontrolledChangelistModuleChanged.AddSP(this, &SSourceControlChangelistsWidget::OnSourceControlStateChanged);
 
-	TreeView = CreateTreeviewWidget();
+	ChangelistTreeView = CreateChangelistTreeView(ChangelistTreeNodes);
+	UncontrolledChangelistTreeView = CreateChangelistTreeView(UncontrolledChangelistTreeNodes);
+	FileTreeView = CreateChangelistFilesView();
+
+	ChangelistExpandableArea = SNew(SExpandableChangelistArea)
+		.HeaderText_Lambda([this]() { return FText::Format(LOCTEXT("SourceControl_ChangeLists", "Changelists ({0})"), ChangelistTreeNodes.Num()); })
+		.ChangelistView(ChangelistTreeView.ToSharedRef())
+		.OnNewChangelist_Lambda([this](){ OnNewChangelist(); return FReply::Handled(); })
+		.OnNewChangelistTooltip(LOCTEXT("Create_New_Changelist", "Create a new changelist."))
+		.SearchButtonVisibility(EVisibility::Collapsed); // Functionality is planned but not fully implemented yet.
+
+	UncontrolledChangelistExpandableArea = SNew(SExpandableChangelistArea)
+		.HeaderText_Lambda([this]() { return FText::Format(LOCTEXT("SourceControl_UncontrolledChangeLists", "Uncontrolled Changelists ({0})"), UncontrolledChangelistTreeNodes.Num()); })
+		.ChangelistView(UncontrolledChangelistTreeView.ToSharedRef())
+		.NewButtonVisibility(EVisibility::Collapsed) // Functionality is planned but not implemented yet.
+		.OnNewChangelistTooltip(LOCTEXT("Create_New_Uncontrolled_Changelist", "Create a new uncontrolled changelist."))
+		.SearchButtonVisibility(EVisibility::Collapsed); // Functionality is planned but not fully implemented yet.
 
 	ChildSlot
 	[
 		SNew(SVerticalBox)
-		+ SVerticalBox::Slot()
+		+SVerticalBox::Slot() // For the toolbar (Refresh button)
 		.AutoHeight()
 		[
 			SNew(SBorder)
@@ -99,44 +352,106 @@ void SSourceControlChangelistsWidget::Construct(const FArguments& InArgs)
 				]
 			]
 		]
-		+ SVerticalBox::Slot()
+		+SVerticalBox::Slot() // Everything below the tools bar: changelist expandable areas + files views + status bar at the bottom
 		[
-			SNew(SScrollBorder, TreeView.ToSharedRef())
-			.Visibility(TAttribute<EVisibility>::Create(TAttribute<EVisibility>::FGetter::CreateLambda([]()->EVisibility { return (ISourceControlModule::Get().IsEnabled() || FUncontrolledChangelistsModule::Get().IsEnabled())
-																																   ? EVisibility::Visible
-																																   : EVisibility::Hidden; })))
+			SNew(SOverlay)
+			+SOverlay::Slot()
 			[
-				TreeView.ToSharedRef()
+				SNew(SBox)
+				.Visibility_Lambda([](){ return !AreChangelistsEnabled() ? EVisibility::Visible: EVisibility::Collapsed; })
+				.HAlign(HAlign_Center)
+				.VAlign(VAlign_Center)
+				[
+					SNew(STextBlock)
+					.Text(LOCTEXT("SourceControl_Disabled", "The source control is disabled or it doesn't support changelists."))
+				]
+			]
+			+SOverlay::Slot()
+			[
+				SNew(SSplitter)
+				.Orientation(EOrientation::Orient_Horizontal)
+				.ResizeMode(ESplitterResizeMode::FixedPosition)
+				.Visibility_Lambda([]() { return AreChangelistsEnabled() ? EVisibility::Visible : EVisibility::Collapsed; })
+
+				// Left slot: Changelists and uncontrolled changelists areas
+				+SSplitter::Slot()
+				.Resizable(true)
+				.Value(0.30)
+				[
+					SNew(SOverlay) // Visible when both Controlled and Uncontrolled changelists are enabled (Need to add a splitter)
+					+SOverlay::Slot()
+					[
+						SNew(SSplitter)
+						.Orientation(EOrientation::Orient_Vertical)
+						.Visibility_Lambda([]() { return AreControlledChangelistsEnabled() && AreUncontrolledChangelistsEnabled() ? EVisibility::Visible : EVisibility::Collapsed; })
+				
+						// Top slot: Changelists
+						+SSplitter::Slot()
+						.SizeRule_Lambda([this](){ return ChangelistExpandableArea->IsExpanded() ? SSplitter::ESizeRule::FractionOfParent : SSplitter::ESizeRule::SizeToContent; })
+						.Value(0.7)
+						[
+							ChangelistExpandableArea.ToSharedRef()
+						]
+
+						// Bottom slot: Uncontrolled Changelists
+						+SSplitter::Slot()
+						.SizeRule_Lambda([this](){ return UncontrolledChangelistExpandableArea->IsExpanded() ? SSplitter::ESizeRule::FractionOfParent : SSplitter::ESizeRule::SizeToContent; })
+						.Value(0.3)
+						[
+							UncontrolledChangelistExpandableArea.ToSharedRef()
+						]
+					]
+					+SOverlay::Slot() // Visibile when controlled changelists are enabled but not the uncontrolled ones.
+					[
+						SNew(SBox)
+						.Visibility_Lambda([](){ return AreControlledChangelistsEnabled() && !AreUncontrolledChangelistsEnabled() ? EVisibility::Visible : EVisibility::Collapsed; })
+						[
+							ChangelistExpandableArea.ToSharedRef()
+						]
+					]
+					+SOverlay::Slot() // Visible when uncontrolled changelist are enabled, but not the controlled ones.
+					[
+						SNew(SBox)
+						.Visibility_Lambda([](){ return !AreControlledChangelistsEnabled() && AreUncontrolledChangelistsEnabled() ? EVisibility::Visible : EVisibility::Collapsed; })
+						[
+							UncontrolledChangelistExpandableArea.ToSharedRef()
+						]
+					]
+				]
+
+				// Right slot: Files associated to the selected the changelist/uncontrolled changelist.
+				+SSplitter::Slot()
+				.Resizable(true)
+				[
+					SNew(SScrollBorder, FileTreeView.ToSharedRef())
+					[
+						FileTreeView.ToSharedRef()
+					]
+				]
 			]
 		]
-		+ SVerticalBox::Slot()
+		+SVerticalBox::Slot() // Status bar (Always visible if uncontrolled changelist are enabled to keep the reconcile status visible at all time)
 		.AutoHeight()
 		[
-			SNew(SBorder)
+			SNew(SBox)
+			.Padding(0, 3)
+			.Visibility_Lambda([this](){ return FUncontrolledChangelistsModule::Get().IsEnabled() || !RefreshStatus.IsEmpty() ? EVisibility::Visible : EVisibility::Collapsed; })
 			[
 				SNew(SHorizontalBox)
-				+ SHorizontalBox::Slot()
+				+SHorizontalBox::Slot()
 				.HAlign(HAlign_Left)
 				.VAlign(VAlign_Center)
 				[
 					SNew(STextBlock)
 					.Text_Lambda([this]() { return RefreshStatus; })
-					.Visibility_Lambda([this]() -> EVisibility
-					{
-						return bIsRefreshing ? EVisibility::Visible : EVisibility::Collapsed;
-					})
 				]
-				+ SHorizontalBox::Slot()
+				+SHorizontalBox::Slot()
 				.HAlign(HAlign_Right)
 				.VAlign(VAlign_Center)
-				.Padding(FMargin(2.f, 0.f))
 				[
 					SNew(STextBlock)
 					.Text_Lambda([]() { return FUncontrolledChangelistsModule::Get().GetReconcileStatus(); })
-					.Visibility_Lambda([]() -> EVisibility
-					{
-						return FUncontrolledChangelistsModule::Get().IsEnabled() ? EVisibility::Visible : EVisibility::Collapsed;
-					})
+					.Visibility_Lambda([]() { return FUncontrolledChangelistsModule::Get().IsEnabled() ? EVisibility::Visible : EVisibility::Collapsed; })
 				]
 			]
 		]
@@ -151,20 +466,11 @@ TSharedRef<SWidget> SSourceControlChangelistsWidget::MakeToolBar()
 
 	ToolBarBuilder.AddToolBarButton(
 		FUIAction(
-			FExecuteAction::CreateLambda([this]() {
-				RequestRefresh();
-				})),
-		NAME_None,
-		LOCTEXT("SourceControl_RefreshButton", "Refresh"),
-		LOCTEXT("SourceControl_RefreshButton_Tooltip", "Refreshes changelists from source control provider."),
-		FSlateIcon(FAppStyle::GetAppStyleSetName(), "SourceControl.Actions.Refresh"));
-
-	ToolBarBuilder.AddToolBarButton(
-		FUIAction(FExecuteAction::CreateSP(this, &SSourceControlChangelistsWidget::OnNewChangelist)),
-		NAME_None,
-		LOCTEXT("SourceControl_NewChangelistButton", "New Changelist"),
-		LOCTEXT("SourceControl_NewChangelistButton_Tooltip", "Creates an empty changelist"),
-		FSlateIcon(FAppStyle::GetAppStyleSetName(), "SourceControl.Actions.Add"));
+			FExecuteAction::CreateLambda([this]() { RequestRefresh(); })),
+			NAME_None,
+			LOCTEXT("SourceControl_RefreshButton", "Refresh"),
+			LOCTEXT("SourceControl_RefreshButton_Tooltip", "Refreshes changelists from source control provider."),
+			FSlateIcon(FAppStyle::GetAppStyleSetName(), "SourceControl.Actions.Refresh"));
 
 	return ToolBarBuilder.MakeWidget();
 }
@@ -177,29 +483,48 @@ bool SSourceControlChangelistsWidget::HasValidationTag(const FText& InChangelist
 	return DescriptionString.Find(ValidationString) != INDEX_NONE;
 }
 
-void SSourceControlChangelistsWidget::EditChangelistDescription(const FText& InNewChangelistDescription, const FSourceControlChangelistStatePtr& InChangelistState) const
+void SSourceControlChangelistsWidget::EditChangelistDescription(const FText& InNewChangelistDescription, const FSourceControlChangelistStatePtr& InChangelistState)
 {
-	auto EditChangelistOperation = ISourceControlOperation::Create<FEditChangelist>();
-	ISourceControlProvider& SourceControlProvider = ISourceControlModule::Get().GetProvider();
-
+	TSharedRef<FEditChangelist> EditChangelistOperation = ISourceControlOperation::Create<FEditChangelist>();
 	EditChangelistOperation->SetDescription(InNewChangelistDescription);
-	SourceControlProvider.Execute(EditChangelistOperation, InChangelistState->GetChangelist());
+	Execute(LOCTEXT("Updating_Changelist_Description", "Updating changelist description..."), EditChangelistOperation, InChangelistState->GetChangelist(), EConcurrency::Asynchronous, FSourceControlOperationComplete::CreateLambda(
+		[](const TSharedRef<ISourceControlOperation>& Operation, ECommandResult::Type InResult)
+		{
+			if (InResult == ECommandResult::Succeeded)
+			{
+				DisplaySourceControlOperationNotification(LOCTEXT("Update_Changelist_Description_Succeeded", "Changelist description successfully updated."), SNotificationItem::CS_Success);
+			}
+			else if (InResult == ECommandResult::Failed)
+			{
+				DisplaySourceControlOperationNotification(LOCTEXT("Update_Changelist_Description_Failed", "Failed to update changelist description."), SNotificationItem::CS_Fail);
+			}
+		}));
 }
 
 void SSourceControlChangelistsWidget::Tick(const FGeometry& AllottedGeometry, const double InCurrentTime, const float InDeltaTime)
 {
+	// Detect transitions of the source control being available/unavailable. Ex: When the user changes the source control in UI, the provider gets selected,
+	// but it is not connected/available until the user accepts the settings. The source control doesn't have callback for availability and we want to refresh everything
+	// once it gets available.
+	if (ISourceControlModule::Get().IsEnabled() && !bSourceControlAvailable && ISourceControlModule::Get().GetProvider().IsAvailable())
+	{
+		bSourceControlAvailable = true;
+		bShouldRefresh = true;
+	}
+
 	if (bShouldRefresh)
 	{
 		if (ISourceControlModule::Get().IsEnabled() || FUncontrolledChangelistsModule::Get().IsEnabled())
 		{
 			RequestRefresh();
-			bShouldRefresh = false;
 		}
 		else
 		{
 			// No provider available, clear changelist tree
 			ClearChangelistsTree();
 		}
+
+		bShouldRefresh = false;
 	}
 	
 	if (bIsRefreshing)
@@ -223,15 +548,16 @@ void SSourceControlChangelistsWidget::RequestRefresh()
 		UpdatePendingChangelistsOperation->SetUpdateShelvedFilesStates(true);
 
 		ISourceControlProvider& SourceControlProvider = ISourceControlModule::Get().GetProvider();
-		SourceControlProvider.Execute(UpdatePendingChangelistsOperation, EConcurrency::Asynchronous);
+		SourceControlProvider.Execute(UpdatePendingChangelistsOperation, EConcurrency::Asynchronous, FSourceControlOperationComplete::CreateSP(this, &SSourceControlChangelistsWidget::OnChangelistsStatusUpdated));
+		OnStartSourceControlOperation(UpdatePendingChangelistsOperation, LOCTEXT("SourceControl_UpdatingChangelist", "Updating changelists..."));
 	}
 
 	if (FUncontrolledChangelistsModule::Get().IsEnabled())
 	{
 		bAnyProviderAvailable = true;
 
-		FUncontrolledChangelistsModule& UncontrolledChangelistModule = FUncontrolledChangelistsModule::Get();
-		UncontrolledChangelistModule.UpdateStatus();
+		// This operation is synchronous and completes right away.
+		FUncontrolledChangelistsModule::Get().UpdateStatus();
 	}
 
 	if (!bAnyProviderAvailable)
@@ -244,14 +570,13 @@ void SSourceControlChangelistsWidget::RequestRefresh()
 void SSourceControlChangelistsWidget::StartRefreshStatus()
 {
 	bIsRefreshing = true;
-	RefreshStatusTimeElapsed = 0;
+	RefreshStatusStartSecs = FPlatformTime::Seconds();
 }
 
 void SSourceControlChangelistsWidget::TickRefreshStatus(double InDeltaTime)
 {
-	RefreshStatusTimeElapsed += InDeltaTime;
-	const int SecondsElapsed = (int)RefreshStatusTimeElapsed;
-	RefreshStatus = FText::Format(LOCTEXT("SourceControl_RefreshStatus", "Refreshing changelists... ({0} s)"), FText::AsNumber(SecondsElapsed));
+	int RefreshStatusTimeElapsed = static_cast<int>(FPlatformTime::Seconds() - RefreshStatusStartSecs);
+	RefreshStatus = FText::Format(LOCTEXT("SourceControl_RefreshStatus", "Refreshing changelists... ({0} s)"), FText::AsNumber(RefreshStatusTimeElapsed));
 }
 
 void SSourceControlChangelistsWidget::EndRefreshStatus()
@@ -261,117 +586,128 @@ void SSourceControlChangelistsWidget::EndRefreshStatus()
 
 void SSourceControlChangelistsWidget::ClearChangelistsTree()
 {
-	if (!ChangelistsNodes.IsEmpty())
+	if (!ChangelistTreeNodes.IsEmpty() || !UncontrolledChangelistTreeNodes.IsEmpty())
 	{
-		ChangelistsNodes.Empty();
-		TreeView->RequestTreeRefresh();
+		ChangelistTreeNodes.Reset();
+		UncontrolledChangelistTreeNodes.Reset();
+		ChangelistTreeView->RequestTreeRefresh();
+		UncontrolledChangelistTreeView->RequestTreeRefresh();
+	}
+
+	if (!FileTreeNodes.IsEmpty())
+	{
+		FileTreeNodes.Reset();
+		FileTreeView->RequestTreeRefresh();
 	}
 }
 
-void SSourceControlChangelistsWidget::Refresh()
+void SSourceControlChangelistsWidget::OnRefresh()
 {
-	if (ISourceControlModule::Get().IsEnabled() || FUncontrolledChangelistsModule::Get().IsEnabled())
+	if (!AreChangelistsEnabled())
 	{
-		TMap<FSourceControlChangelistStateRef, ExpandedState> ExpandedStates;
-		SaveExpandedState(ExpandedStates);
+		ClearChangelistsTree();
+		return;
+	}
 
-		ISourceControlProvider& SourceControlProvider = ISourceControlModule::Get().GetProvider();
-		FUncontrolledChangelistsModule& UncontrolledChangelistModule = FUncontrolledChangelistsModule::Get();
-		TArray<FSourceControlChangelistRef> Changelists = SourceControlProvider.GetChangelists(EStateCacheUsage::Use);
-		TArray<FUncontrolledChangelistStateRef> UncontrolledChangelistStates = UncontrolledChangelistModule.GetChangelistStates();
+	// Views will be teared down and rebuilt from scratch, save the items that are expanded and/or selected to be able to restore those states after the rebuild.
+	FExpandedAndSelectionStates ExpandedAndSelectedStates;
+	SaveExpandedAndSelectionStates(ExpandedAndSelectedStates);
 
-		TArray<FSourceControlChangelistStateRef> ChangelistsStates;
-		SourceControlProvider.GetState(Changelists, ChangelistsStates, EStateCacheUsage::Use);
+	// Query the source control
+	ISourceControlProvider& SourceControlProvider = ISourceControlModule::Get().GetProvider();
+	FUncontrolledChangelistsModule& UncontrolledChangelistModule = FUncontrolledChangelistsModule::Get();
+	TArray<FSourceControlChangelistRef> Changelists = SourceControlProvider.GetChangelists(EStateCacheUsage::Use);
+	TArray<TSharedRef<FUncontrolledChangelistState>> UncontrolledChangelistStates = UncontrolledChangelistModule.GetChangelistStates();
 
-		ChangelistsNodes.Reset(ChangelistsStates.Num());
+	TArray<FSourceControlChangelistStateRef> ChangelistsStates;
+	SourceControlProvider.GetState(Changelists, ChangelistsStates, EStateCacheUsage::Use);
 
-		// Count number of steps for slow task...
-		int32 ElementsToProcess = ChangelistsStates.Num();
-		ElementsToProcess += UncontrolledChangelistStates.Num();
+	// Count number of steps for slow task...
+	int32 ElementsToProcess = ChangelistsStates.Num();
+	ElementsToProcess += UncontrolledChangelistStates.Num();
 
-		for (FSourceControlChangelistStateRef ChangelistState : ChangelistsStates)
-		{
-			ElementsToProcess += ChangelistState->GetFilesStates().Num();
-			ElementsToProcess += ChangelistState->GetShelvedFilesStates().Num();
-		}
+	for (const TSharedRef<ISourceControlChangelistState>& ChangelistState : ChangelistsStates)
+	{
+		ElementsToProcess += ChangelistState->GetFilesStates().Num();
+		ElementsToProcess += ChangelistState->GetShelvedFilesStates().Num();
+	}
 
-		for (FUncontrolledChangelistStateRef UncontrolledChangelistState : UncontrolledChangelistStates)
-		{
-			ElementsToProcess += UncontrolledChangelistState->GetFilesStates().Num();
-			ElementsToProcess += UncontrolledChangelistState->GetOfflineFiles().Num();
-		}
+	for (const TSharedRef<FUncontrolledChangelistState>& UncontrolledChangelistState : UncontrolledChangelistStates)
+	{
+		ElementsToProcess += UncontrolledChangelistState->GetFilesStates().Num();
+		ElementsToProcess += UncontrolledChangelistState->GetOfflineFiles().Num();
+	}
+
+	FScopedSlowTask SlowTask(ElementsToProcess, LOCTEXT("SourceControl_RebuildTree", "Refreshing Tree Items"));
+	SlowTask.MakeDialog(/*bShowCancelButton=*/true);
 
-		FScopedSlowTask SlowTask(ElementsToProcess, LOCTEXT("SourceControl_RebuildTree", "Refreshing Tree Items"));
-		SlowTask.MakeDialog(/*bShowCancelButton=*/true);
+	// Rebuild the tree data models
+	bool bBeautifyPaths = true;
+	ChangelistTreeNodes.Reset(ChangelistsStates.Num());
+	UncontrolledChangelistTreeNodes.Reset(UncontrolledChangelistStates.Num());
+	FileTreeNodes.Reset();
 
-		bool bBeautifyPaths = true;
+	for (const TSharedRef<ISourceControlChangelistState>& ChangelistState : ChangelistsStates)
+	{
+		// Add a changelist.
+		TSharedRef<IChangelistTreeItem> ChangelistNode = MakeShared<FChangelistTreeItem>(ChangelistState);
+		ChangelistTreeNodes.Add(ChangelistNode);
+
+		for (const TSharedRef<ISourceControlState>& FileState : ChangelistState->GetFilesStates())
+		{
+			ChangelistNode->AddChild(MakeShared<FFileTreeItem>(FileState, bBeautifyPaths));
+			SlowTask.EnterProgressFrame();
+			bBeautifyPaths &= !SlowTask.ShouldCancel();
+		}
 
-		for (FSourceControlChangelistStateRef ChangelistState : ChangelistsStates)
+		if (ChangelistState->GetShelvedFilesStates().Num() > 0)
 		{
-			FChangelistTreeItemRef ChangelistTreeItem = MakeShareable(new FChangelistTreeItem(ChangelistState));
+			// Add a shelved files node under the changelist.
+			TSharedRef<IChangelistTreeItem> ShelvedFilesNode = MakeShared<FShelvedChangelistTreeItem>();
+			ChangelistNode->AddChild(ShelvedFilesNode);
 
-			for (FSourceControlStateRef FileRef : ChangelistState->GetFilesStates())
+			for (const TSharedRef<ISourceControlState>& ShelvedFileState : ChangelistState->GetShelvedFilesStates())
 			{
-				FChangelistTreeItemRef FileTreeItem = MakeShareable(new FFileTreeItem(FileRef, bBeautifyPaths));
-				ChangelistTreeItem->AddChild(FileTreeItem);
+				ShelvedFilesNode->AddChild(MakeShared<FShelvedFileTreeItem>(ShelvedFileState, bBeautifyPaths));
 				SlowTask.EnterProgressFrame();
 				bBeautifyPaths &= !SlowTask.ShouldCancel();
 			}
+		}
 
-			if (ChangelistState->GetShelvedFilesStates().Num() > 0)
-			{
-				FChangelistTreeItemRef ShelvedChangelistTreeItem = MakeShareable(new FShelvedChangelistTreeItem());
-				ChangelistTreeItem->AddChild(ShelvedChangelistTreeItem);
+		SlowTask.EnterProgressFrame();
+		bBeautifyPaths &= !SlowTask.ShouldCancel();
+	}
 
-				for (FSourceControlStateRef ShelvedFileRef : ChangelistState->GetShelvedFilesStates())
-				{
-					FChangelistTreeItemRef ShelvedFileTreeItem = MakeShareable(new FShelvedFileTreeItem(ShelvedFileRef, bBeautifyPaths));
-					ShelvedChangelistTreeItem->AddChild(ShelvedFileTreeItem);
-					SlowTask.EnterProgressFrame();
-					bBeautifyPaths &= !SlowTask.ShouldCancel();
-				}
-			}
+	for (const TSharedRef<FUncontrolledChangelistState>& UncontrolledChangelistState : UncontrolledChangelistStates)
+	{
+		// Add an uncontrolled changelist.
+		TSharedRef<IChangelistTreeItem> UncontrolledChangelistNode = MakeShared<FUncontrolledChangelistTreeItem>(UncontrolledChangelistState);
+		UncontrolledChangelistTreeNodes.Add(UncontrolledChangelistNode);
 
-			ChangelistsNodes.Add(ChangelistTreeItem);
+		for (const TSharedRef<ISourceControlState>& FileState : UncontrolledChangelistState->GetFilesStates())
+		{
+			UncontrolledChangelistNode->AddChild(MakeShared<FFileTreeItem>(FileState, bBeautifyPaths));
 			SlowTask.EnterProgressFrame();
 			bBeautifyPaths &= !SlowTask.ShouldCancel();
 		}
 
-		for (FUncontrolledChangelistStateRef UncontrolledChangelistState : UncontrolledChangelistStates)
+		for (const FString& Filename : UncontrolledChangelistState->GetOfflineFiles())
 		{
-			FChangelistTreeItemRef UncontrolledChangelistTreeItem = MakeShareable(new FUncontrolledChangelistTreeItem(UncontrolledChangelistState));
-			
-			for (const FSourceControlStateRef& FileRef : UncontrolledChangelistState->GetFilesStates())
-			{
-				FChangelistTreeItemRef FileTreeItem = MakeShareable(new FFileTreeItem(FileRef, bBeautifyPaths));
-				UncontrolledChangelistTreeItem->AddChild(FileTreeItem);
-				SlowTask.EnterProgressFrame();
-				bBeautifyPaths &= !SlowTask.ShouldCancel();
-			}
-
-			for (const FString& Filename : UncontrolledChangelistState->GetOfflineFiles())
-			{
-				FChangelistTreeItemRef OfflineFileTreeItem = MakeShareable(new FOfflineFileTreeItem(Filename));
-				UncontrolledChangelistTreeItem->AddChild(OfflineFileTreeItem);
-				SlowTask.EnterProgressFrame();
-				bBeautifyPaths &= !SlowTask.ShouldCancel();
-			}
-
-			ChangelistsNodes.Add(UncontrolledChangelistTreeItem);
+			UncontrolledChangelistNode->AddChild(MakeShared<FOfflineFileTreeItem>(Filename));
 			SlowTask.EnterProgressFrame();
 			bBeautifyPaths &= !SlowTask.ShouldCancel();
 		}
 
-		RestoreExpandedState(ExpandedStates);
-
-		TreeView->RequestTreeRefresh();
-	}
-	else
-	{
-		ClearChangelistsTree();
+		SlowTask.EnterProgressFrame();
+		bBeautifyPaths &= !SlowTask.ShouldCancel();
 	}
 
-	EndRefreshStatus();
+	// Views were rebuilt from scratch, try expanding and selecting the nodes that were in that state before the update.
+	RestoreExpandedAndSelectionStates(ExpandedAndSelectedStates);
+
+	ChangelistTreeView->RequestTreeRefresh();
+	UncontrolledChangelistTreeView->RequestTreeRefresh();
+	FileTreeView->RequestTreeRefresh();
 }
 
 void SSourceControlChangelistsWidget::OnSourceControlProviderChanged(ISourceControlProvider& OldProvider, ISourceControlProvider& NewProvider)
@@ -379,17 +715,22 @@ void SSourceControlChangelistsWidget::OnSourceControlProviderChanged(ISourceCont
 	OldProvider.UnregisterSourceControlStateChanged_Handle(SourceControlStateChangedDelegateHandle);
 	SourceControlStateChangedDelegateHandle = NewProvider.RegisterSourceControlStateChanged_Handle(FSourceControlStateChanged::FDelegate::CreateSP(this, &SSourceControlChangelistsWidget::OnSourceControlStateChanged));
 
+	bSourceControlAvailable = NewProvider.IsAvailable(); // Check if it is connected.
 	bShouldRefresh = true;
 }
 
 void SSourceControlChangelistsWidget::OnSourceControlStateChanged()
 {
-	Refresh();
+	// NOTE: No need to call RequestRefresh() to force the SCC to update internal states. We are being invoked because it was update, we just
+	//       need to update the UI to reflect those state changes.
+	OnRefresh();
 }
 
-void SSourceControlChangelistsWidget::OnChangelistsStatusUpdated(const FSourceControlOperationRef& InOperation, ECommandResult::Type InType)
+void SSourceControlChangelistsWidget::OnChangelistsStatusUpdated(const TSharedRef<ISourceControlOperation>& InOperation, ECommandResult::Type InType)
 {
-	Refresh();
+	// NOTE: This is invoked when the 'FUpdatePendingChangelistsStatus' completes. No need to refresh the tree views because OnSourceControlStateChanged() is also called.
+	OnEndSourceControlOperation(InOperation, InType);
+	EndRefreshStatus(); // TODO PL: Need to uniformize all operations status update. The 'Status Update' is different as it displays the time it takes.
 }
 
 void SChangelistTree::Private_SetItemSelection(FChangelistTreeItemPtr TheItem, bool bShouldBeSelected, bool bWasUserDirected)
@@ -419,40 +760,34 @@ void SChangelistTree::Private_SetItemSelection(FChangelistTreeItemPtr TheItem, b
 
 FSourceControlChangelistStatePtr SSourceControlChangelistsWidget::GetCurrentChangelistState()
 {
-	if (!TreeView)
+	if (!ChangelistTreeView)
 	{
 		return nullptr;
 	}
 
-	TArray<FChangelistTreeItemPtr> SelectedItems = TreeView->GetSelectedItems();
-
+	TArray<TSharedPtr<IChangelistTreeItem>> SelectedItems = ChangelistTreeView->GetSelectedItems();
 	if (SelectedItems.Num() != 1 || SelectedItems[0]->GetTreeItemType() != IChangelistTreeItem::Changelist)
 	{
 		return nullptr;
 	}
-	else
-	{
-		return StaticCastSharedPtr<FChangelistTreeItem>(SelectedItems[0])->ChangelistState;
-	}
+
+	return StaticCastSharedPtr<FChangelistTreeItem>(SelectedItems[0])->ChangelistState;
 }
 
 FUncontrolledChangelistStatePtr SSourceControlChangelistsWidget::GetCurrentUncontrolledChangelistState()
 {
-	if (!TreeView)
+	if (!UncontrolledChangelistTreeView)
 	{
 		return nullptr;
 	}
 
-	TArray<FChangelistTreeItemPtr> SelectedItems = TreeView->GetSelectedItems();
-
+	TArray<TSharedPtr<IChangelistTreeItem>> SelectedItems = UncontrolledChangelistTreeView->GetSelectedItems();
 	if (SelectedItems.Num() != 1 || SelectedItems[0]->GetTreeItemType() != IChangelistTreeItem::UncontrolledChangelist)
 	{
 		return nullptr;
 	}
-	else
-	{
-		return StaticCastSharedPtr<FUncontrolledChangelistTreeItem>(SelectedItems[0])->UncontrolledChangelistState;
-	}
+
+	return StaticCastSharedPtr<FUncontrolledChangelistTreeItem>(SelectedItems[0])->UncontrolledChangelistState;
 }
 
 FSourceControlChangelistPtr SSourceControlChangelistsWidget::GetCurrentChangelist()
@@ -463,26 +798,20 @@ FSourceControlChangelistPtr SSourceControlChangelistsWidget::GetCurrentChangelis
 
 FSourceControlChangelistStatePtr SSourceControlChangelistsWidget::GetChangelistStateFromSelection()
 {
-	if (!TreeView)
-	{
-		return nullptr;
-	}
-
-	TArray<FChangelistTreeItemPtr> SelectedItems = TreeView->GetSelectedItems();
-
-	if (SelectedItems.Num() == 0 || SelectedItems[0]->GetTreeItemType() == IChangelistTreeItem::Invalid)
+	TArray<FChangelistTreeItemPtr> SelectedItems = ChangelistTreeView->GetSelectedItems();
+	if (SelectedItems.Num() == 0)
 	{
 		return nullptr;
 	}
 
 	FChangelistTreeItemPtr Item = SelectedItems[0];
-
-	while (Item && Item->GetTreeItemType() != IChangelistTreeItem::Invalid)
+	while (Item)
 	{
 		if (Item->GetTreeItemType() == IChangelistTreeItem::Changelist)
+		{
 			return StaticCastSharedPtr<FChangelistTreeItem>(Item)->ChangelistState;
-		else
-			Item = Item->GetParent();
+		}
+		Item = Item->GetParent();
 	}
 
 	return nullptr;
@@ -496,88 +825,68 @@ FSourceControlChangelistPtr SSourceControlChangelistsWidget::GetChangelistFromSe
 
 TArray<FString> SSourceControlChangelistsWidget::GetSelectedFiles()
 {
-	TArray<FChangelistTreeItemPtr> SelectedItems = TreeView->GetSelectedItems();
+	TArray<FChangelistTreeItemPtr> SelectedItems = FileTreeView->GetSelectedItems();
+	TArray<FString> Files;
 
-	if (SelectedItems.Num() == 0 || SelectedItems[0]->GetTreeItemType() != IChangelistTreeItem::File)
-	{
-		return TArray<FString>();
-	}
-	else
+	for (const TSharedPtr<IChangelistTreeItem>& Item : SelectedItems)
 	{
-		TArray<FString> Files;
-
-		for (FChangelistTreeItemPtr Item : SelectedItems)
+		if (Item->GetTreeItemType() == IChangelistTreeItem::File)
 		{
-			if (Item->GetTreeItemType() != IChangelistTreeItem::File)
-			{
-				continue;
-			}
-
 			Files.Add(StaticCastSharedPtr<FFileTreeItem>(Item)->FileState->GetFilename());
 		}
-
-		return Files;
 	}
+
+	return Files;
 }
 
 void SSourceControlChangelistsWidget::GetSelectedFiles(TArray<FString>& OutControlledFiles, TArray<FString>& OutUncontrolledFiles)
 {
-	TArray<FChangelistTreeItemPtr> SelectedItems = TreeView->GetSelectedItems();
+	TArray<FChangelistTreeItemPtr> SelectedItems = FileTreeView->GetSelectedItems();
 
-	for (const FChangelistTreeItemPtr& Item : SelectedItems)
+	for (const TSharedPtr<IChangelistTreeItem>& Item : SelectedItems)
 	{
 		if (Item->GetTreeItemType() != IChangelistTreeItem::File)
 		{
 			continue;
 		}
 
-		const FChangelistTreeItemPtr& Parent = Item->GetParent();
-
-		if (!Parent.IsValid())
+		if (TSharedPtr<IChangelistTreeItem> Parent = Item->GetParent())
 		{
-			continue;
-		}
-
-		const FString& Filename = StaticCastSharedPtr<FFileTreeItem>(Item)->FileState->GetFilename();
+			const FString& Filename = StaticCastSharedPtr<FFileTreeItem>(Item)->FileState->GetFilename();
 
-		if (Parent->GetTreeItemType() == IChangelistTreeItem::Changelist)
-		{
-			OutControlledFiles.Add(Filename);
-		}
-		else if (Parent->GetTreeItemType() == IChangelistTreeItem::UncontrolledChangelist)
-		{
-			OutUncontrolledFiles.Add(Filename);
+			if (Parent->GetTreeItemType() == IChangelistTreeItem::Changelist)
+			{
+				OutControlledFiles.Add(Filename);
+			}
+			else if (Parent->GetTreeItemType() == IChangelistTreeItem::UncontrolledChangelist)
+			{
+				OutUncontrolledFiles.Add(Filename);
+			}
 		}
 	}
 }
 
 void SSourceControlChangelistsWidget::GetSelectedFileStates(TArray<FSourceControlStateRef>& OutControlledFileStates, TArray<FSourceControlStateRef>& OutUncontrolledFileStates)
 {
-	TArray<FChangelistTreeItemPtr> SelectedItems = TreeView->GetSelectedItems();
+	TArray<TSharedPtr<IChangelistTreeItem>> SelectedItems = FileTreeView->GetSelectedItems();
 
-	for (const FChangelistTreeItemPtr& Item : SelectedItems)
+	for (const TSharedPtr<IChangelistTreeItem>& Item : SelectedItems)
 	{
 		if (Item->GetTreeItemType() != IChangelistTreeItem::File)
 		{
 			continue;
 		}
 
-		const FChangelistTreeItemPtr& Parent = Item->GetParent();
-
-		if (!Parent.IsValid())
-		{
-			continue;
-		}
-
-		FSourceControlStateRef FileState = StaticCastSharedPtr<FFileTreeItem>(Item)->FileState;
-
-		if (Parent->GetTreeItemType() == IChangelistTreeItem::Changelist)
+		if (const TSharedPtr<IChangelistTreeItem>& Parent = Item->GetParent())
 		{
-			OutControlledFileStates.Add(MoveTemp(FileState));
-		}
-		else if (Parent->GetTreeItemType() == IChangelistTreeItem::UncontrolledChangelist)
-		{
-			OutUncontrolledFileStates.Add(MoveTemp(FileState));
+			if (Parent->GetTreeItemType() == IChangelistTreeItem::Changelist)
+			{
+				OutControlledFileStates.Add(StaticCastSharedPtr<FFileTreeItem>(Item)->FileState);
+			}
+			else if (Parent->GetTreeItemType() == IChangelistTreeItem::UncontrolledChangelist)
+			{
+				OutUncontrolledFileStates.Add(StaticCastSharedPtr<FFileTreeItem>(Item)->FileState);
+			}
 		}
 	}
 }
@@ -585,24 +894,33 @@ void SSourceControlChangelistsWidget::GetSelectedFileStates(TArray<FSourceContro
 TArray<FString> SSourceControlChangelistsWidget::GetSelectedShelvedFiles()
 {
 	TArray<FString> ShelvedFiles;
-	TArray<FChangelistTreeItemPtr> SelectedItems = TreeView->GetSelectedItems();
-	
-	if (SelectedItems.Num() > 0)
+
+	for (const TSharedPtr<IChangelistTreeItem>& Item : FileTreeView->GetSelectedItems())
 	{
-		if (SelectedItems[0]->GetTreeItemType() == IChangelistTreeItem::ShelvedChangelist)
+		if (Item->GetTreeItemType() == IChangelistTreeItem::ShelvedFile)
 		{
-			check(SelectedItems.Num() == 1);
-			const TArray<FChangelistTreeItemPtr>& ShelvedChildren = SelectedItems[0]->GetChildren();
-			for (FChangelistTreeItemPtr Item : ShelvedChildren)
-			{
-				ShelvedFiles.Add(StaticCastSharedPtr<FShelvedFileTreeItem>(Item)->FileState->GetFilename());
-			}
+			ShelvedFiles.Add(StaticCastSharedPtr<FShelvedFileTreeItem>(Item)->FileState->GetFilename());
 		}
-		else if (SelectedItems[0]->GetTreeItemType() == IChangelistTreeItem::ShelvedFile)
+	}
+
+	// No individual 'shelved file' selected?
+	if (ShelvedFiles.IsEmpty())
+	{
+		// Check if the user selected the 'Shelved Files' changelist.
+		for (const TSharedPtr<IChangelistTreeItem>& Item : ChangelistTreeView->GetSelectedItems())
 		{
-			for (FChangelistTreeItemPtr Item : SelectedItems)
+			if (Item->GetTreeItemType() == IChangelistTreeItem::ShelvedChangelist)
 			{
-				ShelvedFiles.Add(StaticCastSharedPtr<FShelvedFileTreeItem>(Item)->FileState->GetFilename());
+				// Add all items of the 'Shelved Files' changelist.
+				for (const TSharedPtr<IChangelistTreeItem>& Children : Item->GetChildren())
+				{
+					if (Children->GetTreeItemType() == IChangelistTreeItem::ShelvedFile)
+					{
+						ShelvedFiles.Add(StaticCastSharedPtr<FShelvedFileTreeItem>(Children)->FileState->GetFilename());
+					}
+				}
+
+				break; // UI only allows to select one changelist at the time.
 			}
 		}
 	}
@@ -610,34 +928,65 @@ TArray<FString> SSourceControlChangelistsWidget::GetSelectedShelvedFiles()
 	return ShelvedFiles;
 }
 
-bool SSourceControlChangelistsWidget::IsParentOfSelection(const IChangelistTreeItem::TreeItemType ParentType) const
+void SSourceControlChangelistsWidget::Execute(const FText& Message, const TSharedRef<ISourceControlOperation>& InOperation, const EConcurrency::Type InConcurrency, const FSourceControlOperationComplete& InOperationCompleteDelegate)
 {
-	return Algo::AnyOf(TreeView->GetSelectedItems(), [ParentType = ParentType](const FChangelistTreeItemPtr& Item)
-	{
-		IChangelistTreeItem::TreeItemType ItemType = Item->GetTreeItemType();
+	return Execute(Message, InOperation, nullptr, TArray<FString>(), InConcurrency, InOperationCompleteDelegate);
+}
 
-		if (ItemType == ParentType)
-		{
-			return true;
-		}
-		else if ((ItemType == IChangelistTreeItem::File) || (ItemType == IChangelistTreeItem::ShelvedChangelist))
-		{
-			const FChangelistTreeItemPtr& Parent = Item->GetParent();
-			return Parent.IsValid() && (Parent->GetTreeItemType() == ParentType);
-		}
-		else if (ItemType == IChangelistTreeItem::ShelvedFile)
-		{
-			const FChangelistTreeItemPtr& Parent = Item->GetParent();
+void SSourceControlChangelistsWidget::Execute(const FText& Message, const TSharedRef<ISourceControlOperation>& InOperation, TSharedPtr<ISourceControlChangelist> InChangelist, const EConcurrency::Type InConcurrency, const FSourceControlOperationComplete& InOperationCompleteDelegate)
+{
+	return Execute(Message, InOperation, MoveTemp(InChangelist), TArray<FString>(), InConcurrency, InOperationCompleteDelegate);
+}
+
+void SSourceControlChangelistsWidget::Execute(const FText& Message, const TSharedRef<ISourceControlOperation>& InOperation, const TArray<FString>& InFiles, EConcurrency::Type InConcurrency, const FSourceControlOperationComplete& InOperationCompleteDelegate)
+{
+	return Execute(Message, InOperation, nullptr, InFiles, InConcurrency, InOperationCompleteDelegate);
+}
+
+void SSourceControlChangelistsWidget::Execute(const FText& Message, const TSharedRef<ISourceControlOperation>& InOperation, TSharedPtr<ISourceControlChangelist> InChangelist, const TArray<FString>& InFiles, EConcurrency::Type InConcurrency, const FSourceControlOperationComplete& InOperationCompleteDelegate)
+{
+	ISourceControlProvider& SourceControlProvider = ISourceControlModule::Get().GetProvider();
+
+	// Start the operation.
+	OnStartSourceControlOperation(InOperation, Message);
+
+	if (InConcurrency == EConcurrency::Asynchronous)
+	{
+		// Pass a weak ptr to the lambda to protect in case the 'this' widget is closed/destroyed before the source control operation completes.
+		TWeakPtr<SSourceControlChangelistsWidget> ThisWeak(StaticCastSharedRef<SSourceControlChangelistsWidget>(AsShared()));
 
-			if (Parent.IsValid())
+		SourceControlProvider.Execute(InOperation, MoveTemp(InChangelist), InFiles, InConcurrency, FSourceControlOperationComplete::CreateLambda(
+			[ThisWeak, InOperationCompleteDelegate](const TSharedRef<ISourceControlOperation>& Operation, ECommandResult::Type InResult)
 			{
-				const FChangelistTreeItemPtr& GrandParent = Parent->GetParent();
-				return GrandParent.IsValid() && (GrandParent->GetTreeItemType() == ParentType);
-			}
-		}
+				if (TSharedPtr<SSourceControlChangelistsWidget> ThisPtr = ThisWeak.Pin())
+				{
+					InOperationCompleteDelegate.ExecuteIfBound(Operation, InResult);
+					ThisPtr->OnEndSourceControlOperation(Operation, InResult);
+				}
+			}));
+	}
+	else
+	{
+		FScopedSlowTask Progress(0.f, Message);
+		Progress.MakeDialog();
+		ECommandResult::Type Result = SourceControlProvider.Execute(InOperation, InChangelist, InFiles, InConcurrency, InOperationCompleteDelegate);
+		OnEndSourceControlOperation(InOperation, Result);
+	}
+}
 
-		return false;
-	});
+void SSourceControlChangelistsWidget::ExecuteUncontrolledChangelistOperation(const FText& Message, const TFunction<void()>& UncontrolledOperation)
+{
+	ExecuteUncontrolledChangelistOperationWithSlowTaskWrapper(Message, UncontrolledOperation);
+}
+
+void SSourceControlChangelistsWidget::OnStartSourceControlOperation(TSharedRef<ISourceControlOperation> Operation, const FText& Message)
+{
+	RefreshStatus = Message; // TODO: Should have a queue to stack async operations going on to correctly display concurrent async operations.
+}
+
+void SSourceControlChangelistsWidget::OnEndSourceControlOperation(const TSharedRef<ISourceControlOperation>& Operation, ECommandResult::Type InType)
+{
+	RefreshStatus = FText::GetEmpty(); // TODO: Should have a queue to stack async operations going on to correctly display concurrent async operations.
 }
 
 void SSourceControlChangelistsWidget::OnNewChangelist()
@@ -654,11 +1003,20 @@ void SSourceControlChangelistsWidget::OnNewChangelist()
 		return;
 	}
 
-	ISourceControlProvider& SourceControlProvider = ISourceControlModule::Get().GetProvider();
-	auto NewChangelistOperation = ISourceControlOperation::Create<FNewChangelist>();
+	TSharedRef<FNewChangelist> NewChangelistOperation = ISourceControlOperation::Create<FNewChangelist>();
 	NewChangelistOperation->SetDescription(ChangelistDescription);
-
-	SourceControlProvider.Execute(NewChangelistOperation);
+	Execute(LOCTEXT("Creating_Changelist", "Creating changelist..."), NewChangelistOperation, EConcurrency::Asynchronous, FSourceControlOperationComplete::CreateLambda(
+			[](const TSharedRef<ISourceControlOperation>& Operation, ECommandResult::Type InResult)
+			{
+				if (InResult == ECommandResult::Succeeded)
+				{
+					DisplaySourceControlOperationNotification(LOCTEXT("Create_Changelist_Succeeded", "Changelist successfully created."), SNotificationItem::CS_Success);
+				}
+				else if (InResult == ECommandResult::Failed)
+				{
+					DisplaySourceControlOperationNotification(LOCTEXT("Create_Changelist_Failed", "Failed to create the changelist."), SNotificationItem::CS_Fail);
+				}
+			}));
 }
 
 void SSourceControlChangelistsWidget::OnDeleteChangelist()
@@ -668,8 +1026,20 @@ void SSourceControlChangelistsWidget::OnDeleteChangelist()
 		return;
 	}
 
-	ISourceControlProvider& SourceControlProvider = ISourceControlModule::Get().GetProvider();
-	SourceControlProvider.Execute(ISourceControlOperation::Create<FDeleteChangelist>(), GetCurrentChangelist());
+	TSharedRef<FDeleteChangelist> DeleteChangelistOperation = ISourceControlOperation::Create<FDeleteChangelist>();
+	
+	Execute(LOCTEXT("Deleting_Changelist", "Deleting changelist..."), DeleteChangelistOperation, GetCurrentChangelist(), EConcurrency::Asynchronous, FSourceControlOperationComplete::CreateLambda(
+		[](const TSharedRef<ISourceControlOperation>& Operation, ECommandResult::Type InResult)
+		{
+			if (InResult == ECommandResult::Succeeded)
+			{
+				DisplaySourceControlOperationNotification(LOCTEXT("Delete_Changelist_Succeeded", "Changelist successfully deleted."), SNotificationItem::CS_Success);
+			}
+			else if (InResult == ECommandResult::Failed)
+			{
+				DisplaySourceControlOperationNotification(LOCTEXT("Delete_Changelist_Failed", "Failed to delete the selected changelist."), SNotificationItem::CS_Fail);
+			}
+		}));
 }
 
 bool SSourceControlChangelistsWidget::CanDeleteChangelist()
@@ -705,10 +1075,20 @@ void SSourceControlChangelistsWidget::OnEditChangelist()
 
 void SSourceControlChangelistsWidget::OnRevertUnchanged()
 {
-	ISourceControlProvider& SourceControlProvider = ISourceControlModule::Get().GetProvider();
-
-	auto RevertUnchangedOperation = ISourceControlOperation::Create<FRevertUnchanged>();
-	SourceControlProvider.Execute(RevertUnchangedOperation, GetChangelistFromSelection(), GetSelectedFiles());
+	TSharedRef<FRevertUnchanged> RevertUnchangedOperation = ISourceControlOperation::Create<FRevertUnchanged>();
+	Execute(LOCTEXT("Reverting_Unchanged_Files", "Reverting unchanged file(s)..."), RevertUnchangedOperation, GetChangelistFromSelection(), GetSelectedFiles(), EConcurrency::Synchronous, FSourceControlOperationComplete::CreateLambda(
+		[](const TSharedRef<ISourceControlOperation>& Operation, ECommandResult::Type InType)
+		{
+			// NOTE: This operation message should tell how many files were reverted and how many weren't.
+			if (Operation->GetResultInfo().ErrorMessages.Num() == 0)
+			{
+				DisplaySourceControlOperationNotification(LOCTEXT("Revert_Unchanged_Files_Succeeded", "Unchanged files were reverted."), SNotificationItem::CS_Success);
+			}
+			else if (InType == ECommandResult::Failed)
+			{
+				DisplaySourceControlOperationNotification(LOCTEXT("Revert_Unchanged_Files_Failed", "Failed to revert unchanged files."), SNotificationItem::CS_Fail);
+			}
+		}));
 }
 
 bool SSourceControlChangelistsWidget::CanRevertUnchanged()
@@ -721,7 +1101,13 @@ void SSourceControlChangelistsWidget::OnRevert()
 	FText DialogText;
 	FText DialogTitle;
 
-	const bool bApplyOnChangelist = (GetCurrentChangelist() != nullptr);
+	TArray<FString> SelectedControlledFiles;
+	TArray<FString> SelectedUncontrolledFiles;
+
+	GetSelectedFiles(SelectedControlledFiles, SelectedUncontrolledFiles);
+
+	// Apply to the entire changelist only of there are no files selected.
+	const bool bApplyOnChangelist = (SelectedControlledFiles.Num() == 0 && SelectedUncontrolledFiles.Num() == 0);
 
 	if (bApplyOnChangelist)
 	{
@@ -741,52 +1127,55 @@ void SSourceControlChangelistsWidget::OnRevert()
 		return;
 	}
 
-	TArray<FString> SelectedControlledFiles;
-	TArray<FString> SelectedUncontrolledFiles;
-
-	GetSelectedFiles(SelectedControlledFiles, SelectedUncontrolledFiles);
-
-	FSourceControlChangelistPtr SelectedChangelist = GetChangelistFromSelection();
-	ISourceControlProvider& SourceControlProvider = ISourceControlModule::Get().GetProvider();
-
-	// Reverts the selected Changelist or Files
-	if (SelectedChangelist.IsValid())
+	// Can only have one changelist selected at the time in the left split view (either a 'Changelist' or a 'Uncontrolled Changelist')
+	if (TSharedPtr<ISourceControlChangelist> SelectedChangelist = GetChangelistFromSelection())
 	{
-		TSet<FString> ToRevertFiles;
-		FSourceControlChangelistStatePtr ChangelistState;
-		ChangelistState = SourceControlProvider.GetState(SelectedChangelist.ToSharedRef(), EStateCacheUsage::Use);
+		// No specific files selected, pick all the files in the selected the changelist.
+		if (SelectedControlledFiles.IsEmpty())
+		{
+			// Find all the files in that changelist.
+			if (FSourceControlChangelistStatePtr ChangelistState = ISourceControlModule::Get().GetProvider().GetState(SelectedChangelist.ToSharedRef(), EStateCacheUsage::Use))
+			{
+				Algo::Transform(ChangelistState->GetFilesStates(), SelectedControlledFiles, [](const FSourceControlStateRef& FileState)
+				{
+					return FileState->GetFilename();
+				});
+			}
+		}
 
-		if (ChangelistState.IsValid())
+		if (!SelectedControlledFiles.IsEmpty())
 		{
-			Algo::Transform(ChangelistState->GetFilesStates(), ToRevertFiles, [](FSourceControlStateRef FileState)
+			ExecuteChangelistOperationWithSlowTaskWrapper(LOCTEXT("Reverting_Files", "Reverting file(s)..."), [&SelectedControlledFiles]()
 			{
-				return FileState->GetFilename();
+				if (SourceControlHelpers::RevertAndReloadPackages(SelectedControlledFiles))
+				{
+					DisplaySourceControlOperationNotification(LOCTEXT("Revert_Files_Succeeded", "The selected file(s) were reverted."), SNotificationItem::CS_Success);
+				}
+				else
+				{
+					DisplaySourceControlOperationNotification(LOCTEXT("Revert_Files_Failed", "Failed to revert the selected file(s)."), SNotificationItem::CS_Fail);
+				}
 			});
 		}
-
-		// Enforce unique occurences of files.
-		ToRevertFiles.Append(SelectedControlledFiles);
-		SelectedControlledFiles = ToRevertFiles.Array();
 	}
-
-	if (!SelectedControlledFiles.IsEmpty())
-	{
-		SourceControlHelpers::RevertAndReloadPackages(SelectedControlledFiles);
-	}
-
-	FUncontrolledChangelistStatePtr SelectedUncontrolledChangelist = GetCurrentUncontrolledChangelistState();
-
-	// Reverts the selected Uncontrolled Changelist
-	if (SelectedUncontrolledChangelist.IsValid())
+	else if (TSharedPtr<FUncontrolledChangelistState> SelectedUncontrolledChangelist = GetCurrentUncontrolledChangelistState())
 	{
-		Algo::Transform(SelectedUncontrolledChangelist->GetFilesStates(), SelectedUncontrolledFiles, [](const FSourceControlStateRef& State) { return State->GetFilename(); });
-	}
+		// No individual uncontrolled files were selected, revert all the files from the selected uncontrolled changelist.
+		if (SelectedUncontrolledFiles.IsEmpty())
+		{
+			Algo::Transform(SelectedUncontrolledChangelist->GetFilesStates(), SelectedUncontrolledFiles, [](const FSourceControlStateRef& State) { return State->GetFilename(); });
+		}
 
-	// Reverts selected Uncontrolled Files
-	if (!SelectedUncontrolledFiles.IsEmpty())
-	{
-		FUncontrolledChangelistsModule::Get().OnRevert(SelectedUncontrolledFiles);
+		// Revert uncontrolled files (if any).
+		if (!SelectedUncontrolledFiles.IsEmpty())
+		{
+			ExecuteUncontrolledChangelistOperation(LOCTEXT("Reverting_Uncontrolled_Files", "Reverting uncontrolled files..."), [&SelectedUncontrolledFiles]()
+			{
+				FUncontrolledChangelistsModule::Get().OnRevert(SelectedUncontrolledFiles);
+			});
+		}
 	}
+	// No changelist selected (and consequently, no files displayed that could be selected).
 }
 
 bool SSourceControlChangelistsWidget::CanRevert()
@@ -825,24 +1214,54 @@ void SSourceControlChangelistsWidget::OnShelve()
 		}
 	}
 
-	ISourceControlProvider& SourceControlProvider = ISourceControlModule::Get().GetProvider();
-	auto ShelveOperation = ISourceControlOperation::Create<FShelve>();
+	TSharedRef<FShelve> ShelveOperation = ISourceControlOperation::Create<FShelve>();
 	ShelveOperation->SetDescription(ChangelistDescription);
-	SourceControlProvider.Execute(ShelveOperation, CurrentChangelist->GetChangelist(), GetSelectedFiles());
+	Execute(LOCTEXT("Shelving_Files", "Shelving file(s)..."), ShelveOperation, CurrentChangelist->GetChangelist(), GetSelectedFiles(), EConcurrency::Asynchronous, FSourceControlOperationComplete::CreateLambda(
+		[](const TSharedRef<ISourceControlOperation>& Operation, ECommandResult::Type InResult)
+		{
+			if (InResult == ECommandResult::Succeeded)
+			{
+				DisplaySourceControlOperationNotification(LOCTEXT("Shelve_Files_Succeeded", "The selected file(s) were shelved."), SNotificationItem::CS_Success);
+			}
+			else if (InResult == ECommandResult::Failed)
+			{
+				DisplaySourceControlOperationNotification(LOCTEXT("Shelve_Files_Failed", "Failed to shelved the selected file(s)."), SNotificationItem::CS_Fail);
+			}
+		}));
 }
 
 void SSourceControlChangelistsWidget::OnUnshelve()
 {
-	ISourceControlProvider& SourceControlProvider = ISourceControlModule::Get().GetProvider();
-	auto UnshelveOperation = ISourceControlOperation::Create<FUnshelve>();
-	SourceControlProvider.Execute(UnshelveOperation, GetChangelistFromSelection(), GetSelectedShelvedFiles());
+	TSharedRef<FUnshelve> UnshelveOperation = ISourceControlOperation::Create<FUnshelve>();
+	Execute(LOCTEXT("Unshelving_Files", "Unshelving file(s)..."), UnshelveOperation, GetChangelistFromSelection(), GetSelectedShelvedFiles(), EConcurrency::Synchronous, FSourceControlOperationComplete::CreateLambda(
+		[](const TSharedRef<ISourceControlOperation>& Operation, ECommandResult::Type InResult)
+		{
+			if (InResult == ECommandResult::Succeeded)
+			{
+				DisplaySourceControlOperationNotification(LOCTEXT("Unshelve_Files_Succeeded", "The selected file(s) were unshelved."), SNotificationItem::CS_Success);
+			}
+			else if (InResult == ECommandResult::Failed)
+			{
+				DisplaySourceControlOperationNotification(LOCTEXT("Unshelve_Files_Failed", "Failed to unshelved the selected file(s)."), SNotificationItem::CS_Fail);
+			}
+		}));
 }
 
 void SSourceControlChangelistsWidget::OnDeleteShelvedFiles()
 {
-	ISourceControlProvider& SourceControlProvider = ISourceControlModule::Get().GetProvider();
-	auto DeleteShelvedOperation = ISourceControlOperation::Create<FDeleteShelved>();
-	SourceControlProvider.Execute(DeleteShelvedOperation, GetChangelistFromSelection(), GetSelectedShelvedFiles());
+	TSharedRef<FDeleteShelved> DeleteShelvedOperation = ISourceControlOperation::Create<FDeleteShelved>();
+	Execute(LOCTEXT("Deleting_Shelved_Files", "Deleting shelved file(s)..."), DeleteShelvedOperation, GetChangelistFromSelection(), GetSelectedShelvedFiles(), EConcurrency::Asynchronous, FSourceControlOperationComplete::CreateLambda(
+		[](const TSharedRef<ISourceControlOperation>& Operation, ECommandResult::Type InResult)
+		{
+			if (InResult == ECommandResult::Succeeded)
+			{
+				DisplaySourceControlOperationNotification(LOCTEXT("Delete_Shelved_Files_Succeeded", "The selected shelved file(s) were deleted."), SNotificationItem::CS_Success);
+			}
+			else if (InResult == ECommandResult::Failed)
+			{
+				DisplaySourceControlOperationNotification(LOCTEXT("Delete_Shelved_Files_Failed", "Failed to delete the selected shelved file(s)."), SNotificationItem::CS_Fail);
+			}
+		}));
 }
 
 static bool GetChangelistValidationResult(FSourceControlChangelistPtr InChangelist, FString& OutValidationTitleText, FString& OutValidationWarningsText, FString& OutValidationErrorsText)
@@ -1122,11 +1541,11 @@ void SSourceControlChangelistsWidget::OnMoveFiles()
 
 	// Build selection list for changelists
 	TArray<SSourceControlDescriptionItem> Items;
-	Items.Reset(ChangelistsNodes.Num() + (bAddNewChangelistEntry ? 1 : 0));
+	Items.Reset(ChangelistTreeNodes.Num() + UncontrolledChangelistTreeNodes.Num() + (bAddNewChangelistEntry ? 1 : 0));
 
 	if (bAddNewChangelistEntry)
 	{
-		// First is always new changelist
+		// First item in the 'Move To' list is always 'new changelist'
 		Items.Emplace(
 			LOCTEXT("SourceControl_NewChangelistText", "New Changelist"),
 			LOCTEXT("SourceControl_NewChangelistDescription", "<enter description here>"),
@@ -1135,21 +1554,20 @@ void SSourceControlChangelistsWidget::OnMoveFiles()
 
 	const bool bCanEditAlreadyExistingChangelistDescription = false;
 
-	for (FChangelistTreeItemPtr Changelist : ChangelistsNodes)
+	for (TSharedPtr<IChangelistTreeItem>& Changelist : ChangelistTreeNodes)
 	{
-		if (!Changelist)
-		{
-			continue;
-		}
-
 		if (Changelist->GetTreeItemType() == IChangelistTreeItem::Changelist)
 		{
-			const auto& TypedChangelist = StaticCastSharedPtr<FChangelistTreeItem>(Changelist);
+			const TSharedPtr<FChangelistTreeItem>& TypedChangelist = StaticCastSharedPtr<FChangelistTreeItem>(Changelist);
 			Items.Emplace(TypedChangelist->GetDisplayText(), TypedChangelist->GetDescriptionText(), bCanEditAlreadyExistingChangelistDescription);
 		}
-		else if (Changelist->GetTreeItemType() == IChangelistTreeItem::UncontrolledChangelist)
+	}
+
+	for (TSharedPtr<IChangelistTreeItem>& UncontrolledChangelist : UncontrolledChangelistTreeNodes)
+	{
+		if (UncontrolledChangelist->GetTreeItemType() == IChangelistTreeItem::UncontrolledChangelist)
 		{
-			const FUncontrolledChangelistTreeItemPtr& TypedChangelist = StaticCastSharedPtr<FUncontrolledChangelistTreeItem>(Changelist);
+			const TSharedPtr<FUncontrolledChangelistTreeItem>& TypedChangelist = StaticCastSharedPtr<FUncontrolledChangelistTreeItem>(UncontrolledChangelist);
 			Items.Emplace(TypedChangelist->GetDisplayText(), TypedChangelist->GetDescriptionText(), bCanEditAlreadyExistingChangelistDescription);
 		}
 	}
@@ -1172,40 +1590,76 @@ void SSourceControlChangelistsWidget::OnMoveFiles()
 
 	ISourceControlProvider& SourceControlProvider = ISourceControlModule::Get().GetProvider();
 
-	// Create new changelist
+	// Move files to a new changelist
 	if (bAddNewChangelistEntry && PickedItem == 0)
 	{
-		auto NewChangelistOperation = ISourceControlOperation::Create<FNewChangelist>();
+		// NOTE: To perform async move, we would need to copy the list of selected uncontrolled files and ensure the list wasn't modified when callback occurs. For now run synchronously.
+		TSharedRef<FNewChangelist> NewChangelistOperation = ISourceControlOperation::Create<FNewChangelist>();
 		NewChangelistOperation->SetDescription(ChangelistDescription);
-		SourceControlProvider.Execute(NewChangelistOperation, SelectedControlledFiles);
+		Execute(LOCTEXT("Moving_Files_New_Changelist", "Moving file(s) to a new changelist..."), NewChangelistOperation, SelectedControlledFiles, EConcurrency::Synchronous, FSourceControlOperationComplete::CreateLambda(
+			[this, SelectedUncontrolledFiles](const TSharedRef<ISourceControlOperation>& Operation, ECommandResult::Type InResult)
+			{
+				if (InResult == ECommandResult::Succeeded)
+				{
+					// NOTE: Perform uncontrolled move only if the new changelist was created and the controlled file were move.
+					if ((!SelectedUncontrolledFiles.IsEmpty()) && static_cast<FNewChangelist&>(Operation.Get()).GetNewChangelist().IsValid())
+					{
+						FUncontrolledChangelistsModule::Get().MoveFilesToControlledChangelist(SelectedUncontrolledFiles, static_cast<FNewChangelist&>(Operation.Get()).GetNewChangelist());
+					}
 
-		if ((!SelectedUncontrolledFiles.IsEmpty()) && NewChangelistOperation->GetNewChangelist().IsValid())
-		{
-			FUncontrolledChangelistsModule::Get().MoveFilesToControlledChangelist(SelectedUncontrolledFiles, NewChangelistOperation->GetNewChangelist());
-		}
+					DisplaySourceControlOperationNotification(LOCTEXT("Move_Files_New_Changelist_Succeeded", "Files were successfully moved to a new changelist."), SNotificationItem::CS_Success);
+				}
+				if (InResult == ECommandResult::Failed)
+				{
+					DisplaySourceControlOperationNotification(LOCTEXT("Move_Files_New_Changelist_Failed", "Failed to move the file to the new changelist."), SNotificationItem::CS_Fail);
+				}
+			}));
 	}
-	else
+	else // Move files to an existing changelist or uncontrolled changelist.
 	{
+		// NOTE: The combo box indices are in this order: New changelist, existing changelist(s), existing uncontrolled changelist(s)
+		FChangelistTreeItemPtr MoveDestination;
 		const int32 ChangelistIndex = (bAddNewChangelistEntry ? PickedItem - 1 : PickedItem);
-		const FChangelistTreeItemPtr& SelectedItem = ChangelistsNodes[ChangelistIndex];
 
-		if (SelectedItem->GetTreeItemType() == IChangelistTreeItem::Changelist)
+		if (ChangelistIndex < ChangelistTreeNodes.Num()) // Move files to a changelist
+		{
+			MoveDestination = ChangelistTreeNodes[ChangelistIndex];
+		}
+		else // Move files to an uncontrolled changelist. All uncontrolled CL were listed after the controlled CL in the combo box, compute the offset.
 		{
-			FSourceControlChangelistPtr Changelist = StaticCastSharedPtr<FChangelistTreeItem>(SelectedItem)->ChangelistState->GetChangelist();
+			MoveDestination = UncontrolledChangelistTreeNodes[ChangelistIndex - ChangelistTreeNodes.Num()];
+		}
 
-			if (!SelectedControlledFiles.IsEmpty())
-			{
-				SourceControlProvider.Execute(ISourceControlOperation::Create<FMoveToChangelist>(), Changelist, SelectedControlledFiles);
-			}
+		// Move file to a changelist.
+		if (MoveDestination->GetTreeItemType() == IChangelistTreeItem::Changelist)
+		{
+			FSourceControlChangelistPtr Changelist = StaticCastSharedPtr<FChangelistTreeItem>(MoveDestination)->ChangelistState->GetChangelist();
 
-			if (!SelectedUncontrolledFiles.IsEmpty())
+			if (!SelectedControlledFiles.IsEmpty())
 			{
-				FUncontrolledChangelistsModule::Get().MoveFilesToControlledChangelist(SelectedUncontrolledFiles, Changelist);
+				Execute(LOCTEXT("Moving_File_Between_Changelists", "Moving file(s) to the selected changelist..."), ISourceControlOperation::Create<FMoveToChangelist>(), Changelist, SelectedControlledFiles, EConcurrency::Synchronous, FSourceControlOperationComplete::CreateLambda(
+					[SelectedUncontrolledFiles, Changelist](const TSharedRef<ISourceControlOperation>& Operation, ECommandResult::Type InResult)
+					{
+						if (InResult == ECommandResult::Succeeded)
+						{
+							// Perform an uncontrolled move only if the controlled file were move successfully.
+							if (!SelectedUncontrolledFiles.IsEmpty())
+							{
+								FUncontrolledChangelistsModule::Get().MoveFilesToControlledChangelist(SelectedUncontrolledFiles, Changelist);
+							}
+
+							DisplaySourceControlOperationNotification(LOCTEXT("Move_Files_Between_Changelist_Succeeded", "File(s) successfully moved to the selected changelist."), SNotificationItem::CS_Success);
+						}
+						else if (InResult == ECommandResult::Failed)
+						{
+							DisplaySourceControlOperationNotification(LOCTEXT("Move_Files_Between_Changelist_Failed", "Failed to move the file(s) to the selected changelist."), SNotificationItem::CS_Fail);
+						}
+					}));
 			}
 		}
-		else if (SelectedItem->GetTreeItemType() == IChangelistTreeItem::UncontrolledChangelist)
+		else if (MoveDestination->GetTreeItemType() == IChangelistTreeItem::UncontrolledChangelist)
 		{
-			const FUncontrolledChangelist UncontrolledChangelist = StaticCastSharedPtr<FUncontrolledChangelistTreeItem>(SelectedItem)->UncontrolledChangelistState->Changelist;
+			const FUncontrolledChangelist UncontrolledChangelist = StaticCastSharedPtr<FUncontrolledChangelistTreeItem>(MoveDestination)->UncontrolledChangelistState->Changelist;
 			
 			TArray<FSourceControlStateRef> SelectedControlledFileStates;
 			TArray<FSourceControlStateRef> SelectedUnControlledFileStates;
@@ -1214,7 +1668,10 @@ void SSourceControlChangelistsWidget::OnMoveFiles()
 
 			if ((!SelectedControlledFileStates.IsEmpty()) || (!SelectedUnControlledFileStates.IsEmpty()))
 			{
-				FUncontrolledChangelistsModule::Get().MoveFilesToUncontrolledChangelist(SelectedControlledFileStates, SelectedUnControlledFileStates, UncontrolledChangelist);
+				ExecuteUncontrolledChangelistOperation(LOCTEXT("Moving_Uncontrolled_Changelist_To", "Moving uncontrolled files..."), [&]()
+				{
+					FUncontrolledChangelistsModule::Get().MoveFilesToUncontrolledChangelist(SelectedControlledFileStates, SelectedUnControlledFileStates, UncontrolledChangelist);
+				});
 			}
 		}
 	}
@@ -1223,7 +1680,7 @@ void SSourceControlChangelistsWidget::OnMoveFiles()
 void SSourceControlChangelistsWidget::OnLocateFile()
 {
 	TArray<FAssetData> AssetsToSync;
-	TArray<FChangelistTreeItemPtr> SelectedItems = TreeView->GetSelectedItems();
+	TArray<FChangelistTreeItemPtr> SelectedItems = FileTreeView->GetSelectedItems();
 
 	for (const FChangelistTreeItemPtr& SelectedItem : SelectedItems)
 	{
@@ -1247,7 +1704,7 @@ void SSourceControlChangelistsWidget::OnLocateFile()
 
 bool SSourceControlChangelistsWidget::CanLocateFile()
 {
-	TArray<FChangelistTreeItemPtr> SelectedItems = TreeView->GetSelectedItems();
+	TArray<FChangelistTreeItemPtr> SelectedItems = FileTreeView->GetSelectedItems();
 
 	auto HasAssetData = [](const FChangelistTreeItemPtr& SelectedItem)
 	{
@@ -1292,7 +1749,7 @@ void SSourceControlChangelistsWidget::OnDiffAgainstWorkspace()
 {
 	if (GetSelectedShelvedFiles().Num() > 0)
 	{
-		FSourceControlStateRef FileState = StaticCastSharedPtr<FShelvedFileTreeItem>(TreeView->GetSelectedItems()[0])->FileState;
+		FSourceControlStateRef FileState = StaticCastSharedPtr<FShelvedFileTreeItem>(FileTreeView->GetSelectedItems()[0])->FileState;
 		FSourceControlWindows::DiffAgainstShelvedFile(FileState);
 	}
 }
@@ -1315,62 +1772,88 @@ TSharedPtr<SWidget> SSourceControlChangelistsWidget::OnOpenContextMenu()
 	FToolMenuContext Context;
 	UToolMenu* Menu = ToolMenus->GenerateMenu(MenuName, Context);
 
-	bool bHasSelectedChangelist = (GetCurrentChangelist() != nullptr);
+	TArray<TSharedPtr<IChangelistTreeItem>> SelectedChangelistNodes = ChangelistTreeView->GetSelectedItems();
+	TArray<TSharedPtr<IChangelistTreeItem>> SelectedUncontrolledChangelistNodes = UncontrolledChangelistTreeView->GetSelectedItems();
+
+	bool bHasSelectedChangelist = SelectedChangelistNodes.Num() > 0 &&  SelectedChangelistNodes[0]->GetTreeItemType() == IChangelistTreeItem::Changelist;
+	bool bHasSelectedShelvedChangelistNode = SelectedChangelistNodes.Num() > 0 &&  SelectedChangelistNodes[0]->GetTreeItemType() == IChangelistTreeItem::ShelvedChangelist;
+	bool bHasSelectedUncontrolledChangelist = SelectedUncontrolledChangelistNodes.Num() > 0 &&  SelectedUncontrolledChangelistNodes[0]->GetTreeItemType() == IChangelistTreeItem::UncontrolledChangelist;
 	bool bHasSelectedFiles = (GetSelectedFiles().Num() > 0);
 	bool bHasSelectedShelvedFiles = (GetSelectedShelvedFiles().Num() > 0);
 	bool bHasEmptySelection = (!bHasSelectedChangelist && !bHasSelectedFiles && !bHasSelectedShelvedFiles);
-	bool bIsChangelistParentOfSelection = IsParentOfSelection(IChangelistTreeItem::Changelist);
-	bool bIsUncontrolledChangelistParentOfSelection = IsParentOfSelection(IChangelistTreeItem::UncontrolledChangelist);
 
 	FToolMenuSection& Section = Menu->AddSection("Source Control");
 	
 	// This should appear only on change lists
 	if (bHasSelectedChangelist)
 	{
-		Section.AddMenuEntry("SubmitChangelist", LOCTEXT("SourceControl_SubmitChangelist", "Submit Changelist..."), LOCTEXT("SourceControl_SubmitChangeslit_Tooltip", "Submits a changelist"), FSlateIcon(),
+		Section.AddMenuEntry(
+			"SubmitChangelist",
+			LOCTEXT("SourceControl_SubmitChangelist", "Submit Changelist..."),
+			LOCTEXT("SourceControl_SubmitChangeslit_Tooltip", "Submits a changelist"),
+			FSlateIcon(),
 			FUIAction(
 				FExecuteAction::CreateSP(this, &SSourceControlChangelistsWidget::OnSubmitChangelist),
 				FCanExecuteAction::CreateSP(this, &SSourceControlChangelistsWidget::CanSubmitChangelist)));
 
-		Section.AddMenuEntry("ValidateChangelist", LOCTEXT("SourceControl_ValidateChangelist", "Validate Changelist"), LOCTEXT("SourceControl_ValidateChangeslit_Tooltip", "Validates a changelist"), FSlateIcon(),
-							 FUIAction(
-							 FExecuteAction::CreateSP(this, &SSourceControlChangelistsWidget::OnValidateChangelist),
-							 FCanExecuteAction::CreateSP(this, &SSourceControlChangelistsWidget::CanValidateChangelist)));
-	}
-
-	// This can appear on both files & changelist
-	if (bIsChangelistParentOfSelection)
-	{
-		Section.AddMenuEntry("RevertUnchanged", LOCTEXT("SourceControl_RevertUnchanged", "Revert Unchanged"), LOCTEXT("SourceControl_Revert_Unchanged_Tooltip", "Reverts unchanged files & changelists"), FSlateIcon(),
-							 FUIAction(FExecuteAction::CreateSP(this, &SSourceControlChangelistsWidget::OnRevertUnchanged),
-									   FCanExecuteAction::CreateSP(this, &SSourceControlChangelistsWidget::CanRevertUnchanged)));
+		Section.AddMenuEntry(
+			"ValidateChangelist",
+			LOCTEXT("SourceControl_ValidateChangelist", "Validate Changelist"), LOCTEXT("SourceControl_ValidateChangeslit_Tooltip", "Validates a changelist"),
+			FSlateIcon(),
+			FUIAction(
+				FExecuteAction::CreateSP(this, &SSourceControlChangelistsWidget::OnValidateChangelist), 
+				FCanExecuteAction::CreateSP(this, &SSourceControlChangelistsWidget::CanValidateChangelist)));
+
+		Section.AddMenuEntry(
+			"RevertUnchanged",
+			LOCTEXT("SourceControl_RevertUnchanged", "Revert Unchanged"),
+			LOCTEXT("SourceControl_Revert_Unchanged_Tooltip", "Reverts unchanged files & changelists"),
+			FSlateIcon(),
+			FUIAction(
+				FExecuteAction::CreateSP(this, &SSourceControlChangelistsWidget::OnRevertUnchanged),
+				FCanExecuteAction::CreateSP(this, &SSourceControlChangelistsWidget::CanRevertUnchanged)));
 	}
 
-	if (bIsChangelistParentOfSelection || bIsUncontrolledChangelistParentOfSelection)
+	if (bHasSelectedChangelist || bHasSelectedUncontrolledChangelist)
 	{
-		Section.AddMenuEntry("Revert", LOCTEXT("SourceControl_Revert", "Revert Files"), LOCTEXT("SourceControl_Revert_Tooltip", "Reverts all files in the changelist or from the selection"), FSlateIcon(),
+		Section.AddMenuEntry(
+			"Revert",
+			LOCTEXT("SourceControl_Revert", "Revert Files"),
+			LOCTEXT("SourceControl_Revert_Tooltip", "Reverts all files in the changelist or from the selection"),
+			FSlateIcon(),
 			FUIAction(
 				FExecuteAction::CreateSP(this, &SSourceControlChangelistsWidget::OnRevert),
 				FCanExecuteAction::CreateSP(this, &SSourceControlChangelistsWidget::CanRevert)));
 	}
 
-	if (bIsChangelistParentOfSelection && (bHasSelectedFiles || bHasSelectedShelvedFiles || (bHasSelectedChangelist && (GetCurrentChangelistState()->GetFilesStates().Num() > 0 || GetCurrentChangelistState()->GetShelvedFilesStates().Num() > 0))))
+	if (bHasSelectedChangelist && (bHasSelectedFiles || bHasSelectedShelvedFiles || (bHasSelectedChangelist && (GetCurrentChangelistState()->GetFilesStates().Num() > 0 || GetCurrentChangelistState()->GetShelvedFilesStates().Num() > 0))))
 	{
 		Section.AddSeparator("ShelveSeparator");
 	}
 
-	if (bIsChangelistParentOfSelection && (bHasSelectedFiles || (bHasSelectedChangelist && GetCurrentChangelistState()->GetFilesStates().Num() > 0)))
+	if (bHasSelectedChangelist && (bHasSelectedFiles || (bHasSelectedChangelist && GetCurrentChangelistState()->GetFilesStates().Num() > 0)))
 	{
-		Section.AddMenuEntry("Shelve", LOCTEXT("SourceControl_Shelve", "Shelve Files"), LOCTEXT("SourceControl_Shelve_Tooltip", "Shelves the changelist or the selected files"), FSlateIcon(),
+		Section.AddMenuEntry("Shelve",
+			LOCTEXT("SourceControl_Shelve", "Shelve Files"),
+			LOCTEXT("SourceControl_Shelve_Tooltip", "Shelves the changelist or the selected files"),
+			FSlateIcon(),
 			FUIAction(FExecuteAction::CreateSP(this, &SSourceControlChangelistsWidget::OnShelve)));
 	}
 
-	if (bHasSelectedShelvedFiles || (bHasSelectedChangelist && GetCurrentChangelistState()->GetShelvedFilesStates().Num() > 0))
+	if (bHasSelectedShelvedFiles || bHasSelectedShelvedChangelistNode)
 	{
-		Section.AddMenuEntry("Unshelve", LOCTEXT("SourceControl_Unshelve", "Unshelve Files"), LOCTEXT("SourceControl_Unshelve_Tooltip", "Unshelve selected files or changelist"), FSlateIcon(),
+		Section.AddMenuEntry(
+			"Unshelve",
+			LOCTEXT("SourceControl_Unshelve", "Unshelve Files"),
+			LOCTEXT("SourceControl_Unshelve_Tooltip", "Unshelve selected files or changelist"),
+			FSlateIcon(),
 			FUIAction(FExecuteAction::CreateSP(this, &SSourceControlChangelistsWidget::OnUnshelve)));
 
-		Section.AddMenuEntry("DeleteShelved", LOCTEXT("SourceControl_DeleteShelved", "Delete Shelved Files"), LOCTEXT("SourceControl_DeleteShelved_Tooltip", "Delete selected shelved files or all from changelist"), FSlateIcon(),
+		Section.AddMenuEntry(
+			"DeleteShelved",
+			LOCTEXT("SourceControl_DeleteShelved", "Delete Shelved Files"),
+			LOCTEXT("SourceControl_DeleteShelved_Tooltip", "Delete selected shelved files or all from changelist"),
+			FSlateIcon(),
 			FUIAction(FExecuteAction::CreateSP(this, &SSourceControlChangelistsWidget::OnDeleteShelvedFiles)));
 	}
 
@@ -1378,7 +1861,11 @@ TSharedPtr<SWidget> SSourceControlChangelistsWidget::OnOpenContextMenu()
 	if (bHasSelectedShelvedFiles)
 	{
 		// Diff against workspace
-		Section.AddMenuEntry("DiffAgainstWorkspace", LOCTEXT("SourceControl_DiffAgainstWorkspace", "Diff Against Workspace Files..."), LOCTEXT("SourceControl_DiffAgainstWorkspace_Tooltip", "Diff shelved file against the (local) workspace file"), FSlateIcon(),
+		Section.AddMenuEntry(
+			"DiffAgainstWorkspace",
+			LOCTEXT("SourceControl_DiffAgainstWorkspace", "Diff Against Workspace Files..."),
+			LOCTEXT("SourceControl_DiffAgainstWorkspace_Tooltip", "Diff shelved file against the (local) workspace file"),
+			FSlateIcon(),
 			FUIAction(
 				FExecuteAction::CreateSP(this, &SSourceControlChangelistsWidget::OnDiffAgainstWorkspace),
 				FCanExecuteAction::CreateSP(this, &SSourceControlChangelistsWidget::CanDiffAgainstWorkspace)));
@@ -1389,19 +1876,20 @@ TSharedPtr<SWidget> SSourceControlChangelistsWidget::OnOpenContextMenu()
 		Section.AddSeparator("ChangelistsSeparator");
 	}
 
-	// This should appear only if we have no selection
-	if (bHasEmptySelection)
-	{
-		Section.AddMenuEntry("NewChangelist", LOCTEXT("SourceControl_NewChangelist", "New Changelist..."), LOCTEXT("SourceControl_NewChangelist_Tooltip", "Creates an empty changelist"), FSlateIcon(),
-			FUIAction(FExecuteAction::CreateSP(this, &SSourceControlChangelistsWidget::OnNewChangelist)));
-	}
-
 	if (bHasSelectedChangelist)
 	{
-		Section.AddMenuEntry("EditChangelist", LOCTEXT("SourceControl_EditChangelist", "Edit Changelist..."), LOCTEXT("SourceControl_Edit_Changelist_Tooltip", "Edit a changelist description"), FSlateIcon(),
+		Section.AddMenuEntry(
+			"EditChangelist",
+			LOCTEXT("SourceControl_EditChangelist", "Edit Changelist..."),
+			LOCTEXT("SourceControl_Edit_Changelist_Tooltip", "Edit a changelist description"),
+			FSlateIcon(),
 			FUIAction(FExecuteAction::CreateSP(this, &SSourceControlChangelistsWidget::OnEditChangelist)));
 
-		Section.AddMenuEntry("DeleteChangelist", LOCTEXT("SourceControl_DeleteChangelist", "Delete Empty Changelist"), LOCTEXT("SourceControl_Delete_Changelist_Tooltip", "Deletes an empty changelist"), FSlateIcon(),
+		Section.AddMenuEntry(
+			"DeleteChangelist",
+			LOCTEXT("SourceControl_DeleteChangelist", "Delete Empty Changelist"),
+			LOCTEXT("SourceControl_Delete_Changelist_Tooltip", "Deletes an empty changelist"),
+			FSlateIcon(),
 			FUIAction(
 				FExecuteAction::CreateSP(this, &SSourceControlChangelistsWidget::OnDeleteChangelist),
 				FCanExecuteAction::CreateSP(this, &SSourceControlChangelistsWidget::CanDeleteChangelist)));
@@ -1412,20 +1900,33 @@ TSharedPtr<SWidget> SSourceControlChangelistsWidget::OnOpenContextMenu()
 	{
 		Section.AddSeparator("FilesSeparator");
 
-		Section.AddMenuEntry("MoveFiles", LOCTEXT("SourceControl_MoveFiles", "Move Files To..."), LOCTEXT("SourceControl_MoveFiles_Tooltip", "Move Files To A Different Changelist..."), FSlateIcon(),
-			FUIAction(
-				FExecuteAction::CreateSP(this, &SSourceControlChangelistsWidget::OnMoveFiles)));
-
-		Section.AddMenuEntry("LocateFile", LOCTEXT("SourceControl_LocateFile", "Locate File..."), LOCTEXT("SourceControl_LocateFile_Tooltip", "Locate File in Project..."), FSlateIcon(),
+		Section.AddMenuEntry(
+			"MoveFiles", LOCTEXT("SourceControl_MoveFiles", "Move Files To..."),
+			LOCTEXT("SourceControl_MoveFiles_Tooltip", "Move Files To A Different Changelist..."),
+			FSlateIcon(),
+			FUIAction(FExecuteAction::CreateSP(this, &SSourceControlChangelistsWidget::OnMoveFiles)));
+
+		Section.AddMenuEntry(
+			"LocateFile",
+			LOCTEXT("SourceControl_LocateFile", "Locate File..."),
+			LOCTEXT("SourceControl_LocateFile_Tooltip", "Locate File in Project..."),
+			FSlateIcon(),
 			FUIAction(
 				FExecuteAction::CreateSP(this, &SSourceControlChangelistsWidget::OnLocateFile),
 				FCanExecuteAction::CreateSP(this, &SSourceControlChangelistsWidget::CanLocateFile)));
 
-		Section.AddMenuEntry("ShowHistory", LOCTEXT("SourceControl_ShowHistory", "Show History..."), LOCTEXT("SourceControl_ShowHistory_ToolTip", "Show File History From Selection..."), FSlateIcon(),
-			FUIAction(
-				FExecuteAction::CreateSP(this, &SSourceControlChangelistsWidget::OnShowHistory)));
-
-		Section.AddMenuEntry("DiffAgainstLocalVersion", LOCTEXT("SourceControl_DiffAgainstDepot", "Diff Against Depot..."), LOCTEXT("SourceControl_DiffAgainstLocal_Tooltip", "Diff local file against depot revision."), FSlateIcon(),
+		Section.AddMenuEntry(
+			"ShowHistory",
+			LOCTEXT("SourceControl_ShowHistory", "Show History..."),
+			LOCTEXT("SourceControl_ShowHistory_ToolTip", "Show File History From Selection..."),
+			FSlateIcon(),
+			FUIAction(FExecuteAction::CreateSP(this, &SSourceControlChangelistsWidget::OnShowHistory)));
+
+		Section.AddMenuEntry(
+			"DiffAgainstLocalVersion",
+			LOCTEXT("SourceControl_DiffAgainstDepot", "Diff Against Depot..."),
+			LOCTEXT("SourceControl_DiffAgainstLocal_Tooltip", "Diff local file against depot revision."),
+			FSlateIcon(),
 			FUIAction(
 				FExecuteAction::CreateSP(this, &SSourceControlChangelistsWidget::OnDiffAgainstDepot),
 				FCanExecuteAction::CreateSP(this, &SSourceControlChangelistsWidget::CanDiffAgainstDepot)));
@@ -1435,45 +1936,69 @@ TSharedPtr<SWidget> SSourceControlChangelistsWidget::OnOpenContextMenu()
 	{
 		Section.AddSeparator("ReconcileSeparator");
 
-		Section.AddMenuEntry("Reconcile assets", LOCTEXT("SourceControl_ReconcileAssets", "Reconcile assets"), LOCTEXT("SourceControl_ReconcileAssets_Tooltip", "Look for uncontrolled modification in currently added assets."), FSlateIcon(),
-								FUIAction(FExecuteAction::CreateLambda([]() { FUncontrolledChangelistsModule::Get().OnReconcileAssets(); })));
+		Section.AddMenuEntry(
+			"Reconcile assets",
+			LOCTEXT("SourceControl_ReconcileAssets", "Reconcile assets"),
+			LOCTEXT("SourceControl_ReconcileAssets_Tooltip", "Look for uncontrolled modification in currently added assets."),
+			FSlateIcon(),
+			FUIAction(FExecuteAction::CreateLambda([]() { FUncontrolledChangelistsModule::Get().OnReconcileAssets(); })));
 	}
 
 	return ToolMenus->GenerateWidget(Menu);
 }
 
-TSharedRef<SChangelistTree> SSourceControlChangelistsWidget::CreateTreeviewWidget()
+TSharedRef<SChangelistTree> SSourceControlChangelistsWidget::CreateChangelistTreeView(TArray<TSharedPtr<IChangelistTreeItem>>& ItemSources)
+{
+	return SNew(SChangelistTree)
+		.ItemHeight(24.0f)
+		.TreeItemsSource(&ItemSources)
+		.OnGenerateRow(this, &SSourceControlChangelistsWidget::OnGenerateRow)
+		.OnGetChildren(this, &SSourceControlChangelistsWidget::OnGetChangelistChildren)
+		.SelectionMode(ESelectionMode::Single)
+		.OnContextMenuOpening(this, &SSourceControlChangelistsWidget::OnOpenContextMenu)
+		.OnSelectionChanged(this, &SSourceControlChangelistsWidget::OnChangelistSelectionChanged);
+}
+
+TSharedRef<STreeView<FChangelistTreeItemPtr>> SSourceControlChangelistsWidget::CreateChangelistFilesView()
 {
-	return SAssignNew(TreeView, SChangelistTree)
+	return SNew(STreeView<FChangelistTreeItemPtr>)
 		.ItemHeight(24.0f)
-		.TreeItemsSource(&ChangelistsNodes)
+		.TreeItemsSource(&FileTreeNodes)
 		.OnGenerateRow(this, &SSourceControlChangelistsWidget::OnGenerateRow)
-		.OnGetChildren(this, &SSourceControlChangelistsWidget::OnGetChildren)
+		.OnGetChildren(this, &SSourceControlChangelistsWidget::OnGetFileChildren)
 		.SelectionMode(ESelectionMode::Multi)
 		.OnContextMenuOpening(this, &SSourceControlChangelistsWidget::OnOpenContextMenu)
 		.HeaderRow
 		(
 			SNew(SHeaderRow)
-			+ SHeaderRow::Column("Change")
-			.DefaultLabel(LOCTEXT("Change", "Change"))
+			+SHeaderRow::Column("Icon")
+			.DefaultLabel(FText::GetEmpty())
+			.FillSized(18)
+
+			+SHeaderRow::Column("Name")
+			.DefaultLabel(LOCTEXT("Name", "Name"))
 			.FillWidth(0.2f)
-			+ SHeaderRow::Column("Description")
-			.DefaultLabel(LOCTEXT("Description", "Description"))
+
+			+SHeaderRow::Column("Path")
+			.DefaultLabel(LOCTEXT("Path", "Path"))
 			.FillWidth(0.6f)
-			+ SHeaderRow::Column("Type")
+
+			+SHeaderRow::Column("Type")
 			.DefaultLabel(LOCTEXT("Type", "Type"))
 			.FillWidth(0.2f)
 		);
 }
 
-
-class SChangelistTableRow : public SMultiColumnTableRow<FChangelistTreeItemPtr>
+/** Displays a changed list row (icon, cl number, description) */
+class SChangelistTableRow : public STableRow<TSharedPtr<IChangelistTreeItem>>
 {
 public:
 	SLATE_BEGIN_ARGS(SChangelistTableRow)
 		: _TreeItemToVisualize()
+		, _OnPostDrop()
 	{}
 	SLATE_ARGUMENT(FChangelistTreeItemPtr, TreeItemToVisualize)
+	SLATE_EVENT(FSimpleDelegate, OnPostDrop)
 	SLATE_END_ARGS()
 
 public:
@@ -1485,60 +2010,48 @@ public:
 	void Construct(const FArguments& InArgs, const TSharedRef<STableViewBase>& InOwner)
 	{
 		TreeItem = static_cast<FChangelistTreeItem*>(InArgs._TreeItemToVisualize.Get());
+		OnPostDrop = InArgs._OnPostDrop;
 
-		auto Args = FSuperRowType::FArguments();
-		SMultiColumnTableRow<FChangelistTreeItemPtr>::Construct(Args, InOwner);
-	}
-
-	// SMultiColumnTableRow overrides
-	virtual TSharedRef<SWidget> GenerateWidgetForColumn(const FName& ColumnName) override
-	{
-		if (ColumnName == TEXT("Change"))
-		{
-			const FSlateBrush* IconBrush = (TreeItem != nullptr) ? FAppStyle::GetBrush(TreeItem->ChangelistState->GetSmallIconName())
-																 : FAppStyle::GetBrush("SourceControl.Changelist");
-
-			return SNew(SHorizontalBox)
-				+ SHorizontalBox::Slot()
-					.AutoWidth()
-					.VAlign(VAlign_Center)
-					[
-						SNew(SExpanderArrow, SharedThis(this))
-					]
+		const FSlateBrush* IconBrush = (TreeItem != nullptr) ?
+			FAppStyle::GetBrush(TreeItem->ChangelistState->GetSmallIconName()) :
+			FAppStyle::GetBrush("SourceControl.Changelist");
 
-				+ SHorizontalBox::Slot()
-					.AutoWidth()
-					.VAlign(VAlign_Center)
-					[
-						SNew(SImage)
-						.Image(IconBrush)
-					]
+		SetToolTipText(GetChangelistDescriptionText());
 
-				+ SHorizontalBox::Slot()
-					.AutoWidth()
-					.Padding(2.0f, 0.0f)
-					.VAlign(VAlign_Center)
-					[
-						SNew(STextBlock)
-						.Text(this, &SChangelistTableRow::GetChangelistText)
-					];
-		}
-		else if (ColumnName == TEXT("Description"))
-		{
-			return SNew(SHorizontalBox)
-				+ SHorizontalBox::Slot()
+		STableRow<FChangelistTreeItemPtr>::Construct(
+			STableRow<FChangelistTreeItemPtr>::FArguments()
+			.Style(FAppStyle::Get(), "TableView.Row")
+			.Content()
+			[
+				SNew(SHorizontalBox)
+				+SHorizontalBox::Slot() // Icon
+				.AutoWidth()
+				[
+					SNew(SImage)
+					.Image(IconBrush)
+				]
+				+SHorizontalBox::Slot() // Changelist number.
+				.Padding(2, 0, 0, 0)
+				.AutoWidth()
+				[
+					SNew(STextBlock)
+					.Text(this, &SChangelistTableRow::GetChangelistText)
+				]
+				+SHorizontalBox::Slot() // Files count.
+				.Padding(4, 0, 4, 0)
+				.AutoWidth()
+				[
+					SNew(STextBlock)
+					.Text(FText::Format(INVTEXT("({0})"), TreeItem->GetFileCount()))
+				]
+				+SHorizontalBox::Slot()
+				.Padding(2, 0, 0, 0)
 				.AutoWidth()
-				.Padding(2.0f, 0.0f)
-				.VAlign(VAlign_Center)
 				[
 					SNew(STextBlock)
 					.Text(this, &SChangelistTableRow::GetChangelistDescriptionText)
-				];
-		}
-		else
-		{
-			return SNullWidget::NullWidget;
-		}
+				]
+			], InOwner);
 	}
 
 	FText GetChangelistText() const
@@ -1560,20 +2073,48 @@ protected:
 	//~ Begin STableRow Interface.
 	virtual FReply OnDrop(const FGeometry& InGeometry, const FDragDropEvent& InDragDropEvent) override
 	{
-		TSharedPtr<FSCCFileDragDropOp> Operation = InDragDropEvent.GetOperationAs<FSCCFileDragDropOp>();
-		if (Operation.IsValid())
+		TSharedPtr<FSCCFileDragDropOp> DropOperation = InDragDropEvent.GetOperationAs<FSCCFileDragDropOp>();
+		if (DropOperation.IsValid())
 		{
-			FSourceControlChangelistPtr Changelist = TreeItem->ChangelistState->GetChangelist();
-			check(Changelist.IsValid());
+			FSourceControlChangelistPtr DestChangelist = TreeItem->ChangelistState->GetChangelist();
+			check(DestChangelist.IsValid());
 
-			TArray<FString> Files;
-			Algo::Transform(Operation->Files, Files, [](const FSourceControlStateRef& State) { return State->GetFilename(); });
+			// NOTE: The UI don't show 'source controlled files' and 'uncontrolled files' at the same time. User cannot select and drag/drop both file types at the same time.
+			if (!DropOperation->Files.IsEmpty())
+			{
+				TArray<FString> Files;
+				Algo::Transform(DropOperation->Files, Files, [](const FSourceControlStateRef& State) { return State->GetFilename(); });
 
-			ISourceControlProvider& SourceControlProvider = ISourceControlModule::Get().GetProvider();
-			FUncontrolledChangelistsModule& UncontrolledChangelistModule = FUncontrolledChangelistsModule::Get();
-			
-			SourceControlProvider.Execute(ISourceControlOperation::Create<FMoveToChangelist>(), Changelist, Files);
-			UncontrolledChangelistModule.MoveFilesToControlledChangelist(Operation->UncontrolledFiles, Changelist);
+				FScopedSlowTask Progress(0.f, LOCTEXT("Dropping_Files_On_Changelist", "Moving file(s) to the selected changelist..."));
+				Progress.MakeDialog();
+
+				ISourceControlProvider& SourceControlProvider = ISourceControlModule::Get().GetProvider();
+				SourceControlProvider.Execute(ISourceControlOperation::Create<FMoveToChangelist>(), DestChangelist, Files, EConcurrency::Synchronous, FSourceControlOperationComplete::CreateLambda(
+					[](const TSharedRef<ISourceControlOperation>& Operation, ECommandResult::Type InResult)
+					{
+						if (InResult == ECommandResult::Succeeded)
+						{
+							DisplaySourceControlOperationNotification(LOCTEXT("Drop_Files_On_Changelist_Succeeded", "File(s) successfully moved to the selected changelist."), SNotificationItem::CS_Success);
+						}
+						else if (InResult == ECommandResult::Failed)
+						{
+							DisplaySourceControlOperationNotification(LOCTEXT("Drop_Files_On_Changelist_Failed", "Failed to move the file(s) to the selected changelist."), SNotificationItem::CS_Fail);
+						}
+					}));
+			}
+			else if (!DropOperation->UncontrolledFiles.IsEmpty())
+			{
+				// NOTE: This function does several operations that can fails but we don't get feedback.
+				ExecuteUncontrolledChangelistOperationWithSlowTaskWrapper(LOCTEXT("Dropping_Uncontrolled_Files_On_Changelist", "Moving uncontrolled file(s) to the selected changelist..."), 
+					[&DropOperation, &DestChangelist]()
+				{
+					FUncontrolledChangelistsModule::Get().MoveFilesToControlledChangelist(DropOperation->UncontrolledFiles, DestChangelist);
+					
+					// TODO: Fix MoveFilesToControlledChangelist() to report the possible errors and display a notification.
+				});
+
+				OnPostDrop.ExecuteIfBound();
+			}
 		}
 
 		return FReply::Handled();
@@ -1583,16 +2124,22 @@ protected:
 private:
 	/** The info about the widget that we are visualizing. */
 	FChangelistTreeItem* TreeItem;
+
+	/** Delegate invoked once the drag and drop operation finished. */
+	FSimpleDelegate OnPostDrop;
 };
 
-class SUncontrolledChangelistTableRow : public SMultiColumnTableRow<FChangelistTreeItemPtr>
+/** Displays an uncontrolled changed list (icon, cl number, description) */
+class SUncontrolledChangelistTableRow : public STableRow<FChangelistTreeItemPtr>
 {
 public:
 	SLATE_BEGIN_ARGS(SUncontrolledChangelistTableRow)
 		: _TreeItemToVisualize()
+		, _OnPostDrop()
 	{
 	}
 	SLATE_ARGUMENT(FChangelistTreeItemPtr, TreeItemToVisualize)
+	SLATE_EVENT(FSimpleDelegate, OnPostDrop)
 	SLATE_END_ARGS()
 
 public:
@@ -1604,47 +2151,42 @@ public:
 	void Construct(const FArguments& InArgs, const TSharedRef<STableViewBase>& InOwner)
 	{
 		TreeItem = static_cast<FUncontrolledChangelistTreeItem*>(InArgs._TreeItemToVisualize.Get());
+		OnPostDrop = InArgs._OnPostDrop;
 
-		auto Args = FSuperRowType::FArguments();
-		SMultiColumnTableRow<FChangelistTreeItemPtr>::Construct(Args, InOwner);
-	}
+			const FSlateBrush* IconBrush = (TreeItem != nullptr) ?
+				FAppStyle::GetBrush(TreeItem->UncontrolledChangelistState->GetSmallIconName()) :
+				FAppStyle::GetBrush("SourceControl.Changelist");
 
-	// SMultiColumnTableRow overrides
-	virtual TSharedRef<SWidget> GenerateWidgetForColumn(const FName& ColumnName) override
-	{
-		if (ColumnName == TEXT("Change"))
-		{
-			const FSlateBrush* IconBrush = (TreeItem != nullptr) ? FAppStyle::GetBrush(TreeItem->UncontrolledChangelistState->GetSmallIconName())
-				: FAppStyle::GetBrush("SourceControl.Changelist");
+		SetToolTipText(GetChangelistDescriptionText());
 
-			return SNew(SHorizontalBox)
+		STableRow<FChangelistTreeItemPtr>::Construct(
+			STableRow<FChangelistTreeItemPtr>::FArguments()
+			.Style(FAppStyle::Get(), "TableView.Row")
+			.Content()
+			[
+				SNew(SHorizontalBox)
 				+ SHorizontalBox::Slot()
-				.AutoWidth()
-				.VAlign(VAlign_Center)
-				[
-					SNew(SExpanderArrow, SharedThis(this))
-				]
-
-			+ SHorizontalBox::Slot()
 				.AutoWidth()
 				.VAlign(VAlign_Center)
 				[
 					SNew(SImage)
 					.Image(IconBrush)
 				]
-
-			+ SHorizontalBox::Slot()
+				+ SHorizontalBox::Slot()
 				.AutoWidth()
 				.Padding(2.0f, 0.0f)
 				.VAlign(VAlign_Center)
 				[
 					SNew(STextBlock)
 					.Text(this, &SUncontrolledChangelistTableRow::GetChangelistText)
-				];
-		}
-		else if (ColumnName == TEXT("Description"))
-		{
-			return SNew(SHorizontalBox)
+				]
+				+SHorizontalBox::Slot() // Files/Offline file count.
+				.Padding(4, 0, 4, 0)
+				.AutoWidth()
+				[
+					SNew(STextBlock)
+					.Text(FText::Format(INVTEXT("({0})"), TreeItem->GetFileCount() + TreeItem->GetOfflineFileCount()))
+				]
 				+ SHorizontalBox::Slot()
 				.AutoWidth()
 				.Padding(2.0f, 0.0f)
@@ -1652,12 +2194,8 @@ public:
 				[
 					SNew(STextBlock)
 					.Text(this, &SUncontrolledChangelistTableRow::GetChangelistDescriptionText)
-				];
-		}
-		else
-		{
-			return SNullWidget::NullWidget;
-		}
+				]
+			], InOwner);
 	}
 
 	FText GetChangelistText() const
@@ -1683,7 +2221,13 @@ protected:
 		
 		if (Operation.IsValid())
 		{
-			FUncontrolledChangelistsModule::Get().MoveFilesToUncontrolledChangelist(Operation->Files, Operation->UncontrolledFiles, TreeItem->UncontrolledChangelistState->Changelist);
+			ExecuteUncontrolledChangelistOperationWithSlowTaskWrapper(LOCTEXT("Drag_File_To_Uncontrolled_Changelist", "Moving file(s) to the selected uncontrolled changelists..."),
+				[this, &Operation]()
+			{
+				FUncontrolledChangelistsModule::Get().MoveFilesToUncontrolledChangelist(Operation->Files, Operation->UncontrolledFiles, TreeItem->UncontrolledChangelistState->Changelist);
+			});
+
+			OnPostDrop.ExecuteIfBound();
 		}
 
 		return FReply::Handled();
@@ -1693,8 +2237,51 @@ protected:
 private:
 	/** The info about the widget that we are visualizing. */
 	FUncontrolledChangelistTreeItem* TreeItem;
+
+	/** Invoked once a drag and drop operation completes. */
+	FSimpleDelegate OnPostDrop;
+};
+
+/** Display the shelved files group node. It displays 'Shelved Files (x)' where X is the nubmer of file shelved. */
+class SShelvedFilesTableRow : public STableRow<TSharedPtr<IChangelistTreeItem>>
+{
+public:
+	SLATE_BEGIN_ARGS(SShelvedFilesTableRow)
+		: _Icon(nullptr)
+	{
+	}
+		SLATE_ARGUMENT(const FSlateBrush*, Icon)
+		SLATE_ARGUMENT(FText, Text)
+	SLATE_END_ARGS()
+
+	void Construct(const FArguments& InArgs, const TSharedRef<STableViewBase>& InOwnerTableView)
+	{
+		STableRow<FChangelistTreeItemPtr>::Construct(
+			STableRow<FChangelistTreeItemPtr>::FArguments()
+			.Content()
+			[
+				SNew(SHorizontalBox)
+					+SHorizontalBox::Slot()
+					.AutoWidth()
+					.VAlign(VAlign_Center)
+					.Padding(5, 0, 0, 0)
+					[
+						SNew(SImage)
+						.Image(InArgs._Icon)
+					]
+					+SHorizontalBox::Slot()
+					.Padding(2.0f, 1.0f)
+					.VAlign(VAlign_Center)
+					[
+						SNew(STextBlock)
+						.Text(InArgs._Text)
+					]
+			],
+			InOwnerTableView);
+	}
 };
 
+/** Display information about a file (icon, name, location, type, etc.) */
 class SFileTableRow : public SMultiColumnTableRow<FChangelistTreeItemPtr>
 {
 public:
@@ -1715,7 +2302,7 @@ public:
 	{
 		TreeItem = static_cast<FFileTreeItem*>(InArgs._TreeItemToVisualize.Get());
 
-		auto Args = FSuperRowType::FArguments()
+		FSuperRowType::FArguments Args = FSuperRowType::FArguments()
 			.OnDragDetected(InArgs._OnDragDetected)
 			.ShowSelection(true);
 		FSuperRowType::Construct(Args, InOwner);
@@ -1724,29 +2311,21 @@ public:
 	// SMultiColumnTableRow overrides
 	virtual TSharedRef<SWidget> GenerateWidgetForColumn(const FName& ColumnName) override
 	{
-		if (ColumnName == TEXT("Change")) // eq. to name
+		if (ColumnName == TEXT("Icon"))
 		{
-			const int32 LeftOffset = (TreeItem->IsShelved() ? 60 : 40);
-
-			return SNew(SHorizontalBox)
-
-			// Icon
-			+ SHorizontalBox::Slot()
-				.AutoWidth()
-				.Padding(LeftOffset, 0, 4, 0)
+			return SNew(SBox)
+				.WidthOverride(16) // Small Icons are usually 16x16
+				.HAlign(HAlign_Center)
 				[
 					SSourceControlCommon::GetSCCFileWidget(TreeItem->FileState, TreeItem->IsShelved())
-				]
-
-			+ SHorizontalBox::Slot()
-				.AutoWidth()
-				.VAlign(VAlign_Center)
-				[
-					SNew(STextBlock)
-					.Text(this, &SFileTableRow::GetDisplayName)
 				];
 		}
-		else if (ColumnName == TEXT("Description")) // eq. to path
+		else if (ColumnName == TEXT("Name"))
+		{
+			return SNew(STextBlock)
+				.Text(this, &SFileTableRow::GetDisplayName);
+		}
+		else if (ColumnName == TEXT("Path"))
 		{
 			return SNew(STextBlock)
 				.Text(this, &SFileTableRow::GetDisplayPath)
@@ -1809,6 +2388,7 @@ private:
 	FFileTreeItem* TreeItem;
 };
 
+/** Display information about an offline file (icon, name, location, type, etc.). */
 class SOfflineFileTableRow : public SMultiColumnTableRow<FChangelistTreeItemPtr>
 {
 public:
@@ -1829,35 +2409,29 @@ public:
 	{
 		TreeItem = static_cast<FOfflineFileTreeItem*>(InArgs._TreeItemToVisualize.Get());
 
-		auto Args = FSuperRowType::FArguments().ShowSelection(true);
+		FSuperRowType::FArguments Args = FSuperRowType::FArguments().ShowSelection(true);
 		FSuperRowType::Construct(Args, InOwner);
 	}
 
 	// SMultiColumnTableRow overrides
 	virtual TSharedRef<SWidget> GenerateWidgetForColumn(const FName& ColumnName) override
 	{
-		if (ColumnName == TEXT("Change")) // eq. to name
+		if (ColumnName == TEXT("Icon"))
 		{
-			return SNew(SHorizontalBox)
-
-				// Icon
-				+ SHorizontalBox::Slot()
-				.AutoWidth()
-				.Padding(40, 0, 4, 0)
+			return SNew(SBox)
+				.WidthOverride(16) // Small Icons are usually 16x16
+				.HAlign(HAlign_Center)
 				[
 					SNew(SImage)
 					.Image(FAppStyle::GetBrush(FName("SourceControl.OfflineFile_Small")))
-				]
-
-			+ SHorizontalBox::Slot()
-				.AutoWidth()
-				.VAlign(VAlign_Center)
-				[
-					SNew(STextBlock)
-					.Text(this, &SOfflineFileTableRow::GetDisplayName)
 				];
 		}
-		else if (ColumnName == TEXT("Description")) // eq. to path
+		else if (ColumnName == TEXT("Name"))
+		{
+			return SNew(STextBlock)
+				.Text(this, &SOfflineFileTableRow::GetDisplayName);
+		}
+		else if (ColumnName == TEXT("Path"))
 		{
 			return SNew(STextBlock)
 				.Text(this, &SOfflineFileTableRow::GetDisplayPath)
@@ -1900,96 +2474,25 @@ public:
 		return TreeItem->GetDisplayColor();
 	}
 
-protected:
-	//~ Begin STableRow Interface.
-
-	//~ End STableRow Interface.
-
 private:
 	/** The info about the widget that we are visualizing. */
 	FOfflineFileTreeItem* TreeItem;
 };
 
-class SShelvedChangelistTableRow : public SMultiColumnTableRow<FChangelistTreeItemPtr>
-{
-public:
-	SLATE_BEGIN_ARGS(SShelvedChangelistTableRow)
-		: _TreeItemToVisualize()
-	{}
-	SLATE_ARGUMENT(FChangelistTreeItemPtr, TreeItemToVisualize)
-		SLATE_END_ARGS()
-
-public:
-	/**
-	* Construct child widgets that comprise this widget.
-	*
-	* @param InArgs Declaration from which to construct this widget.
-	*/
-	void Construct(const FArguments& InArgs, const TSharedRef<STableViewBase>& InOwner)
-	{
-		TreeItem = static_cast<FShelvedChangelistTreeItem*>(InArgs._TreeItemToVisualize.Get());
-
-		auto Args = FSuperRowType::FArguments();
-		SMultiColumnTableRow<FChangelistTreeItemPtr>::Construct(Args, InOwner);
-	}
-
-	// SMultiColumnTableRow overrides
-	virtual TSharedRef<SWidget> GenerateWidgetForColumn(const FName& ColumnName) override
-	{
-		if (ColumnName == TEXT("Change"))
-		{
-			return SNew(SHorizontalBox)
-				+ SHorizontalBox::Slot()
-					.AutoWidth()
-					.VAlign(VAlign_Center)
-					.Padding(5, 0, 4, 0)
-					[
-						SNew(SExpanderArrow, SharedThis(this))
-					]
-				+ SHorizontalBox::Slot()
-					.AutoWidth()
-					.VAlign(VAlign_Center)
-					.Padding(5, 0, 0, 0)
-					[
-						SNew(SImage)
-						.Image(FAppStyle::GetBrush("SourceControl.ShelvedChangelist"))
-					]
-				+ SHorizontalBox::Slot()
-					.Padding(2.0f, 0.0f)
-					.VAlign(VAlign_Center)
-					[
-						SNew(STextBlock)
-						.Text(this, &SShelvedChangelistTableRow::GetText)
-					];
-		}
-		else
-		{
-			return SNullWidget::NullWidget;
-		}
-	}
-
-protected:
-	FText GetText() const
-	{
-		return TreeItem->GetDisplayText();
-	}
-
-private:
-	/** The info about the widget that we are visualizing. */
-	FShelvedChangelistTreeItem* TreeItem;
-};
 
-TSharedRef<ITableRow> SSourceControlChangelistsWidget::OnGenerateRow(FChangelistTreeItemPtr InTreeItem, const TSharedRef<STableViewBase>& OwnerTable)
+TSharedRef<ITableRow> SSourceControlChangelistsWidget::OnGenerateRow(TSharedPtr<IChangelistTreeItem> InTreeItem, const TSharedRef<STableViewBase>& OwnerTable)
 {
 	switch (InTreeItem->GetTreeItemType())
 	{
 	case IChangelistTreeItem::Changelist:
 		return SNew(SChangelistTableRow, OwnerTable)
-			.TreeItemToVisualize(InTreeItem);
+			.TreeItemToVisualize(InTreeItem)
+			.OnPostDrop(this, &SSourceControlChangelistsWidget::OnRefresh);
 
 	case IChangelistTreeItem::UncontrolledChangelist:
 		return SNew(SUncontrolledChangelistTableRow, OwnerTable)
-			.TreeItemToVisualize(InTreeItem);
+			.TreeItemToVisualize(InTreeItem)
+			.OnPostDrop(this, &SSourceControlChangelistsWidget::OnRefresh);
 
 	case IChangelistTreeItem::File:
 		return SNew(SFileTableRow, OwnerTable)
@@ -2001,8 +2504,9 @@ TSharedRef<ITableRow> SSourceControlChangelistsWidget::OnGenerateRow(FChangelist
 			.TreeItemToVisualize(InTreeItem);
 
 	case IChangelistTreeItem::ShelvedChangelist:
-		return SNew(SShelvedChangelistTableRow, OwnerTable)
-			.TreeItemToVisualize(InTreeItem);
+		return SNew(SShelvedFilesTableRow, OwnerTable)
+			.Icon(FAppStyle::GetBrush("SourceControl.ShelvedChangelist"))
+			.Text(static_cast<const FShelvedChangelistTreeItem*>(InTreeItem.Get())->GetDisplayText());
 
 	case IChangelistTreeItem::ShelvedFile:
 		return SNew(SFileTableRow, OwnerTable)
@@ -2017,15 +2521,15 @@ TSharedRef<ITableRow> SSourceControlChangelistsWidget::OnGenerateRow(FChangelist
 
 FReply SSourceControlChangelistsWidget::OnFilesDragged(const FGeometry& InGeometry, const FPointerEvent& InMouseEvent)
 {
-	if (InMouseEvent.IsMouseButtonDown(EKeys::LeftMouseButton) && !TreeView->GetSelectedItems().IsEmpty())
+	if (InMouseEvent.IsMouseButtonDown(EKeys::LeftMouseButton) && !FileTreeView->GetSelectedItems().IsEmpty())
 	{
-		TSharedRef<FSCCFileDragDropOp> Operation = MakeShareable(new FSCCFileDragDropOp());
+		TSharedRef<FSCCFileDragDropOp> Operation = MakeShared<FSCCFileDragDropOp>();
 
-		for (FChangelistTreeItemPtr InTreeItem : TreeView->GetSelectedItems())
+		for (FChangelistTreeItemPtr InTreeItem : FileTreeView->GetSelectedItems())
 		{
 			if (InTreeItem->GetTreeItemType() == IChangelistTreeItem::File)
 			{
-				FFileTreeItemRef FileTreeItem = StaticCastSharedRef<FFileTreeItem>(InTreeItem.ToSharedRef());
+				TSharedRef<FFileTreeItem> FileTreeItem = StaticCastSharedRef<FFileTreeItem>(InTreeItem.ToSharedRef());
 				FSourceControlStateRef FileState = FileTreeItem->FileState;
 
 				if (FileTreeItem->GetParent()->GetTreeItemType() == IChangelistTreeItem::UncontrolledChangelist)
@@ -2047,73 +2551,228 @@ FReply SSourceControlChangelistsWidget::OnFilesDragged(const FGeometry& InGeomet
 	return FReply::Unhandled();
 }
 
-void SSourceControlChangelistsWidget::OnGetChildren(FChangelistTreeItemPtr InParent, TArray<FChangelistTreeItemPtr>& OutChildren)
+void SSourceControlChangelistsWidget::OnGetFileChildren(TSharedPtr<IChangelistTreeItem> InParent, TArray<FChangelistTreeItemPtr>& OutChildren)
 {
-	for (auto& Child : InParent->GetChildren())
-	{
-		// Should never have bogus entries in this list
-		check(Child.IsValid());
-		OutChildren.Add(Child);
-	}
+	// Files are leave and don't have children.
 }
 
-void SSourceControlChangelistsWidget::SaveExpandedState(TMap<FSourceControlChangelistStateRef, ExpandedState>& ExpandedStates) const
+void SSourceControlChangelistsWidget::OnGetChangelistChildren(FChangelistTreeItemPtr InParent, TArray<FChangelistTreeItemPtr>& OutChildren)
 {
-	for (FChangelistTreeItemPtr Root : ChangelistsNodes)
+	if (InParent->GetTreeItemType() == IChangelistTreeItem::Changelist)
 	{
-		if ((Root->GetTreeItemType() != IChangelistTreeItem::Changelist) && (Root->GetTreeItemType() != IChangelistTreeItem::UncontrolledChangelist))
+		// In the data model, a changelist has files as children, but in UI, only the 'Shelved Files' node is displayed under the changelist,
+		// and the files are displayed in the file view at the right.
+		for (const TSharedPtr<IChangelistTreeItem>& Child : InParent->GetChildren())
 		{
-			continue;
+			if (Child->GetTreeItemType() == IChangelistTreeItem::ShelvedChangelist)
+			{
+				if (Child->GetChildren().Num() > 0)
+				{
+					OutChildren.Add(Child); // Add the 'Shelved Files' only if there are shelved files.
+					break; // Found the only possible child for the UI.
+				}
+			}
 		}
+	}
+	else if (InParent->GetTreeItemType() == IChangelistTreeItem::UncontrolledChangelist)
+	{
+		// Uncontrolled changelist nodes do not have children at the moment.
+	}
+}
 
-		bool bChangelistExpanded = TreeView->IsItemExpanded(Root);
+void SSourceControlChangelistsWidget::OnChangelistSelectionChanged(TSharedPtr<IChangelistTreeItem> SelectedChangelist, ESelectInfo::Type SelectionType)
+{
+	FileTreeNodes.Reset();
 
-		bool bShelveExpanded = false;
-		for (FChangelistTreeItemPtr Child : Root->GetChildren())
+	// Add the children of the parent item to the file tree node.
+	auto AddChangelistFilesToFileView = [this](TSharedPtr<IChangelistTreeItem> ParentItem, IChangelistTreeItem::TreeItemType DesiredChildrenType)
+	{
+		for (const TSharedPtr<IChangelistTreeItem>& Child : ParentItem->GetChildren())
 		{
-			if (Child->GetTreeItemType() == IChangelistTreeItem::ShelvedChangelist)
+			if (Child->GetTreeItemType() == DesiredChildrenType)
 			{
-				bShelveExpanded = TreeView->IsItemExpanded(Child);
-				break;
+				FileTreeNodes.Add(Child);
 			}
 		}
+	};
+
+	if (SelectedChangelist) // Can be a Changelist, Uncontrolled Changelist or Shelved Changelist
+	{
+		IChangelistTreeItem::TreeItemType ChangelistType = SelectedChangelist->GetTreeItemType();
+		switch (ChangelistType)
+		{
+			case IChangelistTreeItem::Changelist:
+			case IChangelistTreeItem::ShelvedChangelist:
+				UncontrolledChangelistTreeView->ClearSelection(); // Don't have a changelists selected at the same time than an uncontrolled one, they share the same file view.
+				AddChangelistFilesToFileView(SelectedChangelist, ChangelistType == IChangelistTreeItem::Changelist ? IChangelistTreeItem::File : IChangelistTreeItem::ShelvedFile);
+				break;
+
+			case IChangelistTreeItem::UncontrolledChangelist:
+				ChangelistTreeView->ClearSelection();
+				AddChangelistFilesToFileView(SelectedChangelist, IChangelistTreeItem::File);
+				AddChangelistFilesToFileView(SelectedChangelist, IChangelistTreeItem::OfflineFile);
+				break;
+
+			default:
+				break;
+		}
+	}
 
-		ExpandedState State;
-		State.bChangelistExpanded = bChangelistExpanded;
-		State.bShelveExpanded = bShelveExpanded;
+	FileTreeView->RequestTreeRefresh();
+}
 
-		ExpandedStates.Add(StaticCastSharedPtr<FChangelistTreeItem>(Root)->ChangelistState, State);
+void SSourceControlChangelistsWidget::SaveExpandedAndSelectionStates(FExpandedAndSelectionStates& OutStates)
+{
+	// Save the selected item from the 'changelists' tree.
+	TArray<TSharedPtr<IChangelistTreeItem>> SelectedChangelistItems = ChangelistTreeView->GetSelectedItems();
+	OutStates.SelectedChangelistNode = SelectedChangelistItems.IsEmpty() ? TSharedPtr<IChangelistTreeItem>() : SelectedChangelistItems[0];
+	OutStates.bShelvedFilesNodeSelected = false;
+	if (OutStates.SelectedChangelistNode && OutStates.SelectedChangelistNode->GetTreeItemType() == IChangelistTreeItem::ShelvedChangelist)
+	{
+		OutStates.SelectedChangelistNode = OutStates.SelectedChangelistNode->GetParent();
+		OutStates.bShelvedFilesNodeSelected = true;
 	}
+
+	// Save the selected item from 'uncontrolled changelists' tree.
+	SelectedChangelistItems = UncontrolledChangelistTreeView->GetSelectedItems();
+	OutStates.SelectedUncontrolledChangelistNode = SelectedChangelistItems.IsEmpty() ? TSharedPtr<IChangelistTreeItem>() : SelectedChangelistItems[0];
+
+	// Remember the expanded nodes.
+	check(OutStates.ExpandedTreeNodes.IsEmpty());
+	ChangelistTreeView->GetExpandedItems(OutStates.ExpandedTreeNodes);
+	UncontrolledChangelistTreeView->GetExpandedItems(OutStates.ExpandedTreeNodes);
+
+	// Remember the selected files.
+	OutStates.SelectedFileNodes.Reset(FileTreeView->GetNumItemsSelected());
+	OutStates.SelectedFileNodes.Append(FileTreeView->GetSelectedItems());
 }
 
-void SSourceControlChangelistsWidget::RestoreExpandedState(const TMap<FSourceControlChangelistStateRef, ExpandedState>& ExpandedStates)
+void SSourceControlChangelistsWidget::RestoreExpandedAndSelectionStates(const FExpandedAndSelectionStates& InStates)
 {
-	for (FChangelistTreeItemPtr Root : ChangelistsNodes)
+	// Returns whether two changelist nodes represent the same changelist.
+	auto ChangelistEquals = [](const TSharedPtr<IChangelistTreeItem>& Lhs, const TSharedPtr<IChangelistTreeItem>& Rhs)
 	{
-		if ((Root->GetTreeItemType() != IChangelistTreeItem::Changelist) && (Root->GetTreeItemType() != IChangelistTreeItem::UncontrolledChangelist))
+		// NOTE: This TRUSTS the source control to return the same ' state' pointer before and after an update if the changelist still exists.
+		return static_cast<FChangelistTreeItem*>(Lhs.Get())->ChangelistState == static_cast<FChangelistTreeItem*>(Rhs.Get())->ChangelistState;
+	};
+
+	// Returns whether two uncontrolled changelist nodes represent the same changelist.
+	auto UncontrolledChangelistEquals = [](const TSharedPtr<IChangelistTreeItem>& Lhs, const TSharedPtr<IChangelistTreeItem>& Rhs)
+	{
+		// NOTE: This TRUSTS the source control to return the same 'state' pointer before and after an update if the changelist still exists.
+		return static_cast<FUncontrolledChangelistTreeItem*>(Lhs.Get())->UncontrolledChangelistState == static_cast<FUncontrolledChangelistTreeItem*>(Rhs.Get())->UncontrolledChangelistState;
+	};
+
+	// Find a specified item in a list. The nodes were deleted and recreated during the update and this function is used to match the new node corresponding to the old node.
+	auto Find = [](const TArray<TSharedPtr<IChangelistTreeItem>>& Nodes, const TSharedPtr<IChangelistTreeItem> SearchedItem,
+					const TFunction<bool(const TSharedPtr<IChangelistTreeItem>& Lhs, const TSharedPtr<IChangelistTreeItem>& Rhs)>& Predicate)
+	{
+		if (const TSharedPtr<IChangelistTreeItem>* Node = Nodes.FindByPredicate(
+			[&SearchedItem, &Predicate](const TSharedPtr<IChangelistTreeItem>& Candidate) { return Predicate(SearchedItem, Candidate); }))
 		{
-			continue;
+			return *Node;
 		}
+		return TSharedPtr<IChangelistTreeItem>(); // return nullptr;
+	};
 
-		FSourceControlChangelistStateRef ChangelistState = StaticCastSharedPtr<FChangelistTreeItem>(Root)->ChangelistState;
-		const ExpandedState* State = ExpandedStates.Find(ChangelistState);
+	// Restore the expansion states (Tree is only one level deep)
+	for (const TSharedPtr<IChangelistTreeItem>& ExpandedNode : InStates.ExpandedTreeNodes)
+	{
+		if (ExpandedNode->GetTreeItemType() == IChangelistTreeItem::Changelist)
+		{
+			// Check if the node still exist after the update.
+			if (TSharedPtr<IChangelistTreeItem> MatchingNode = Find(ChangelistTreeNodes, ExpandedNode, ChangelistEquals))
+			{
+				ChangelistTreeView->SetItemExpansion(MatchingNode, true);
+			}
+		}
+		else if (ExpandedNode->GetTreeItemType() == IChangelistTreeItem::UncontrolledChangelist)
+		{
+			// Check if the node still exist after the update.
+			if (TSharedPtr<IChangelistTreeItem> MatchingNode = Find(UncontrolledChangelistTreeNodes, ExpandedNode, UncontrolledChangelistEquals))
+			{
+				UncontrolledChangelistTreeView->SetItemExpansion(MatchingNode, true);
+			}
+		}
+	}
 
-		if (!State)
+	// Restore the selected nodes.
+	if (InStates.SelectedChangelistNode)
+	{
+		if (TSharedPtr<IChangelistTreeItem> MatchingNode = Find(ChangelistTreeNodes, InStates.SelectedChangelistNode, ChangelistEquals))
 		{
-			continue;
+			if (InStates.bShelvedFilesNodeSelected && static_cast<const FChangelistTreeItem*>(MatchingNode.Get())->GetShelvedFileCount() > 0)
+			{
+				for (const TSharedPtr<IChangelistTreeItem>& Child : MatchingNode->GetChildren())
+				{
+					if (Child->GetTreeItemType() == IChangelistTreeItem::ShelvedChangelist)
+					{
+						ChangelistTreeView->SetSelection(Child); // Select 'Shelved Files' node under the changelist.
+						break;
+					}
+				}
+			}
+			else
+			{
+				ChangelistTreeView->SetSelection(MatchingNode); // Select the 'changelist' node
+			}
+		}
+	}
+	else if (InStates.SelectedUncontrolledChangelistNode)
+	{
+		if (TSharedPtr<IChangelistTreeItem> MatchingNode = Find(UncontrolledChangelistTreeNodes, InStates.SelectedUncontrolledChangelistNode, UncontrolledChangelistEquals))
+		{
+			UncontrolledChangelistTreeView->SetSelection(MatchingNode); // Select the 'uncontrolled changelist' node
 		}
+	}
 
-		TreeView->SetItemExpansion(Root, State->bChangelistExpanded);
+	FileTreeView->ClearSelection();
 
-		for (FChangelistTreeItemPtr Child : Root->GetChildren())
+	// Try to reselect the files.
+	for (const TSharedPtr<IChangelistTreeItem>& FileNode : FileTreeNodes)
+	{
+		switch (FileNode->GetTreeItemType())
 		{
-			if (Child->GetTreeItemType() == IChangelistTreeItem::ShelvedChangelist)
-			{
-				TreeView->SetItemExpansion(Child, State->bShelveExpanded);
+			case IChangelistTreeItem::File:
+				if (InStates.SelectedFileNodes.ContainsByPredicate(
+						[&FileNode](const TSharedPtr<IChangelistTreeItem>& Candidate)
+						{
+							return Candidate->GetTreeItemType() == IChangelistTreeItem::File &&
+								static_cast<const FFileTreeItem*>(Candidate.Get())->GetAssetPath().EqualTo(static_cast<const FFileTreeItem*>(FileNode.Get())->GetAssetPath()) &&
+								static_cast<const FFileTreeItem*>(Candidate.Get())->GetFileName().EqualTo(static_cast<const FFileTreeItem*>(FileNode.Get())->GetFileName());
+						}))
+				{
+					FileTreeView->SetItemSelection(FileNode, true);
+				}
+				break;
+
+			case IChangelistTreeItem::ShelvedFile:
+				if (InStates.SelectedFileNodes.ContainsByPredicate(
+						[&FileNode](const TSharedPtr<IChangelistTreeItem>& Candidate)
+						{
+							return Candidate->GetTreeItemType() == IChangelistTreeItem::ShelvedFile &&
+								static_cast<const FShelvedFileTreeItem*>(Candidate.Get())->GetAssetPath().EqualTo(static_cast<const FShelvedFileTreeItem*>(FileNode.Get())->GetAssetPath()) &&
+								static_cast<const FShelvedFileTreeItem*>(Candidate.Get())->GetFileName().EqualTo(static_cast<const FShelvedFileTreeItem*>(FileNode.Get())->GetFileName());
+						}))
+				{
+					FileTreeView->SetItemSelection(FileNode, true);
+				}
+				break;
+
+			case IChangelistTreeItem::OfflineFile:
+				if (InStates.SelectedFileNodes.ContainsByPredicate(
+						[&FileNode](const TSharedPtr<IChangelistTreeItem>& Candidate)
+						{
+							return Candidate->GetTreeItemType() == IChangelistTreeItem::OfflineFile &&
+								static_cast<const FOfflineFileTreeItem*>(Candidate.Get())->GetDisplayPath().EqualTo(static_cast<const FOfflineFileTreeItem*>(FileNode.Get())->GetDisplayPath()) &&
+								static_cast<const FOfflineFileTreeItem*>(Candidate.Get())->GetDisplayName().EqualTo(static_cast<const FOfflineFileTreeItem*>(FileNode.Get())->GetDisplayName());
+						}))
+				{
+					FileTreeView->SetItemSelection(FileNode, true);
+				}
 				break;
-			}
 		}
 	}
 }
 
-#undef LOCTEXT_NAMESPACE
\ No newline at end of file
+#undef LOCTEXT_NAMESPACE
diff --git a/Engine/Source/Editor/SourceControlWindows/Private/SSourceControlChangelists.h b/Engine/Source/Editor/SourceControlWindows/Private/SSourceControlChangelists.h
--- a/Engine/Source/Editor/SourceControlWindows/Private/SSourceControlChangelists.h
+++ b/Engine/Source/Editor/SourceControlWindows/Private/SSourceControlChangelists.h
@@ -10,36 +10,51 @@
 #include "ISourceControlProvider.h"
 #include "SSourceControlCommon.h"
 
+class FChangelistGroupTreeItem;
+class SExpandableChangelistArea;
+
 class SChangelistTree : public STreeView<FChangelistTreeItemPtr>
 {
 private:
 	virtual void Private_SetItemSelection(FChangelistTreeItemPtr TheItem, bool bShouldBeSelected, bool bWasUserDirected = false) override;
 };
 
+
+/**
+ * Displays the user source control change lists.
+ */
 class SSourceControlChangelistsWidget : public SCompoundWidget
 {
 public:
 	SLATE_BEGIN_ARGS(SSourceControlChangelistsWidget) {}
 	SLATE_END_ARGS()
 
-	/**
-	* Constructor.
-	*/
-	SSourceControlChangelistsWidget();
-
 	/** Constructs the widget */
 	void Construct(const FArguments& InArgs);
 
 private:
-	TSharedRef<SChangelistTree> CreateTreeviewWidget();
+	// Holds the list/state of selected and expanded items in the changelist views and file views.
+	struct FExpandedAndSelectionStates
+	{
+		TSharedPtr<IChangelistTreeItem> SelectedChangelistNode;
+		TSharedPtr<IChangelistTreeItem> SelectedUncontrolledChangelistNode;
+		TArray<TSharedPtr<IChangelistTreeItem>> SelectedFileNodes;
+		TSet<TSharedPtr<IChangelistTreeItem>> ExpandedTreeNodes;
+		bool bShelvedFilesNodeSelected = false;
+	};
+
+private:
+	TSharedRef<SChangelistTree> CreateChangelistTreeView(TArray<TSharedPtr<IChangelistTreeItem>>& ItemSources);
+	TSharedRef<STreeView<FChangelistTreeItemPtr>> CreateChangelistFilesView();
 
 	TSharedRef<ITableRow> OnGenerateRow(FChangelistTreeItemPtr InTreeItem, const TSharedRef<STableViewBase>& OwnerTable);
-	void OnGetChildren(FChangelistTreeItemPtr InParent, TArray<FChangelistTreeItemPtr>& OutChildren);
+	void OnGetFileChildren(FChangelistTreeItemPtr InParent, TArray<FChangelistTreeItemPtr>& OutChildren);
+	void OnGetChangelistChildren(FChangelistTreeItemPtr InParent, TArray<FChangelistTreeItemPtr>& OutChildren);
 
 	FReply OnFilesDragged(const FGeometry& InGeometry, const FPointerEvent& InMouseEvent);
 
 	void RequestRefresh();
-	void Refresh();
+	void OnRefresh();
 	void ClearChangelistsTree();
 
 	TSharedPtr<SWidget> OnOpenContextMenu();
@@ -71,13 +86,6 @@ private:
 	/** Returns list of currently selected shelved files */
 	TArray<FString> GetSelectedShelvedFiles();
 
-	/**
-	 * Check if the type given as argument is a parent of selected items.
-	 * @param 	ParentType 	The parent type to look for.
-	 * @return 	True of ParentType is a parent of selected items.
-	 */
-	bool IsParentOfSelection(const IChangelistTreeItem::TreeItemType ParentType) const;
-
 	/** Changelist operations */
 	void OnNewChangelist();
 	void OnDeleteChangelist();
@@ -111,21 +119,26 @@ private:
 	void OnDiffAgainstWorkspace();
 	bool CanDiffAgainstWorkspace();
 
-	void OnSourceControlStateChanged();
+	/** Source control callbacks */
 	void OnSourceControlProviderChanged(ISourceControlProvider& OldProvider, ISourceControlProvider& NewProvider);
-	void OnChangelistsStatusUpdated(const FSourceControlOperationRef& InOperation, ECommandResult::Type InType);
+	void OnSourceControlStateChanged();
+	void OnChangelistSelectionChanged(TSharedPtr<IChangelistTreeItem> SelectedItem, ESelectInfo::Type SelectionType);
+	void OnChangelistsStatusUpdated(const TSharedRef<ISourceControlOperation>& InOperation, ECommandResult::Type InType);
 
-	virtual void Tick(const FGeometry& AllottedGeometry, const double InCurrentTime, const float InDeltaTime) override;
+	void OnStartSourceControlOperation(TSharedRef<ISourceControlOperation> Operation, const FText& Message);
+	void OnEndSourceControlOperation(const TSharedRef<ISourceControlOperation>& Operation, ECommandResult::Type InType);
 
-private:
-	struct ExpandedState
-	{
-		bool bChangelistExpanded;
-		bool bShelveExpanded;
-	};
+	// Wrapper functions on top of the source control ones to display slow tasks for synchronous operations or toast notifications for async ones.
+	void Execute(const FText& Message, const TSharedRef<ISourceControlOperation>& InOperation, TSharedPtr<ISourceControlChangelist> InChangelist, const TArray<FString>& InFiles, EConcurrency::Type InConcurrency, const FSourceControlOperationComplete& InOperationCompleteDelegate);
+	void Execute(const FText& Message, const TSharedRef<ISourceControlOperation>& InOperation, TSharedPtr<ISourceControlChangelist> InChangelist, const EConcurrency::Type InConcurrency, const FSourceControlOperationComplete& InOperationCompleteDelegate);
+	void Execute(const FText& Message, const TSharedRef<ISourceControlOperation>& InOperation, const EConcurrency::Type InConcurrency, const FSourceControlOperationComplete& InOperationCompleteDelegate);
+	void Execute(const FText& Message, const TSharedRef<ISourceControlOperation>& InOperation, const TArray<FString>& InFiles, EConcurrency::Type InConcurrency, const FSourceControlOperationComplete& InOperationCompleteDelegate);
+	void ExecuteUncontrolledChangelistOperation(const FText& Message, const TFunction<void()>& UncontrolledChangelistTask);
 
-	void SaveExpandedState(TMap<FSourceControlChangelistStateRef, ExpandedState>& ExpandedStates) const;
-	void RestoreExpandedState(const TMap<FSourceControlChangelistStateRef, ExpandedState>& ExpandedStates);
+	virtual void Tick(const FGeometry& AllottedGeometry, const double InCurrentTime, const float InDeltaTime) override;
+
+	void SaveExpandedAndSelectionStates(FExpandedAndSelectionStates& OutStates);
+	void RestoreExpandedAndSelectionStates(const FExpandedAndSelectionStates& InStates);
 
 	TSharedRef<SWidget> MakeToolBar();
 
@@ -143,17 +156,28 @@ private:
 	bool HasValidationTag(const FText& InChangelistDescription) const;
 
 	/** Executes an operation to updates the changelist description of the provided changelist with a new description. */
-	void EditChangelistDescription(const FText& InNewChangelistDescription, const FSourceControlChangelistStatePtr& InChangelistState) const;
+	void EditChangelistDescription(const FText& InNewChangelistDescription, const FSourceControlChangelistStatePtr& InChangelistState);
 
 private:
 	/** Tag to append to a changelist that passed validation */
 	static const FText ChangelistValidatedTag;
 
-	/** Changelists (root nodes) */
-	TArray<FChangelistTreeItemPtr> ChangelistsNodes;
+	TSharedPtr<SExpandableChangelistArea> ChangelistExpandableArea;
+	TSharedPtr<SExpandableChangelistArea> UncontrolledChangelistExpandableArea;
+
+	/** Hold the nodes displayed by the changelist tree. */
+	TArray<TSharedPtr<IChangelistTreeItem>> ChangelistTreeNodes;
+	TArray<TSharedPtr<IChangelistTreeItem>> UncontrolledChangelistTreeNodes;
+
+	/** Hold the nodes displayed by the file tree. */
+	TArray<TSharedPtr<IChangelistTreeItem>> FileTreeNodes;
+
+	/** Display the changelists, uncontrolled changelists and shelved nodes. */
+	TSharedPtr<SChangelistTree> ChangelistTreeView;
+	TSharedPtr<SChangelistTree> UncontrolledChangelistTreeView;
 
-	/** Changelist treeview widget */
-	TSharedPtr<SChangelistTree> TreeView;
+	/** Display the list of files associated to the selected changelist, uncontrolled changelist or shelved node. */
+	TSharedPtr<STreeView<FChangelistTreeItemPtr>> FileTreeView;
 
 	/** Source control state changed delegate handle */
 	FDelegateHandle SourceControlStateChangedDelegateHandle;
@@ -161,13 +185,14 @@ private:
 	/** Uncontrolled Changelist changed delegate handle */
 	FDelegateHandle UncontrolledChangelistChangedDelegateHandle;
 
-	bool bShouldRefresh;
+	bool bShouldRefresh = false;
+	bool bSourceControlAvailable = false;
 
 	void StartRefreshStatus();
 	void TickRefreshStatus(double InDeltaTime);
 	void EndRefreshStatus();
 
 	FText RefreshStatus;
-	bool bIsRefreshing;
-	double RefreshStatusTimeElapsed;
+	bool bIsRefreshing = false;
+	double RefreshStatusStartSecs;
 };
diff --git a/Engine/Source/Editor/SourceControlWindows/Private/SSourceControlCommon.cpp b/Engine/Source/Editor/SourceControlWindows/Private/SSourceControlCommon.cpp
--- a/Engine/Source/Editor/SourceControlWindows/Private/SSourceControlCommon.cpp
+++ b/Engine/Source/Editor/SourceControlWindows/Private/SSourceControlCommon.cpp
@@ -21,6 +21,7 @@
 #define LOCTEXT_NAMESPACE "SourceControlChangelist"
 
 //////////////////////////////////////////////////////////////////////////
+
 FChangelistTreeItemPtr IChangelistTreeItem::GetParent() const
 {
 	return Parent;
@@ -31,13 +32,13 @@ const TArray<FChangelistTreeItemPtr>& IChangelistTreeItem::GetChildren() const
 	return Children;
 }
 
-void IChangelistTreeItem::AddChild(FChangelistTreeItemRef Child)
+void IChangelistTreeItem::AddChild(TSharedRef<IChangelistTreeItem> Child)
 {
 	Child->Parent = AsShared();
 	Children.Add(MoveTemp(Child));
 }
 
-void IChangelistTreeItem::RemoveChild(const FChangelistTreeItemRef& Child)
+void IChangelistTreeItem::RemoveChild(const TSharedRef<IChangelistTreeItem>& Child)
 {
 	if (Children.Remove(Child))
 	{
@@ -46,19 +47,14 @@ void IChangelistTreeItem::RemoveChild(const FChangelistTreeItemRef& Child)
 }
 
 //////////////////////////////////////////////////////////////////////////
-FText FShelvedChangelistTreeItem::GetDisplayText() const
-{
-	return LOCTEXT("SourceControl_ShelvedFiles", "Shelved Items");
-}
 
-//////////////////////////////////////////////////////////////////////////
 FFileTreeItem::FFileTreeItem(FSourceControlStateRef InFileState, bool bBeautifyPaths, bool bIsShelvedFile)
-	: FileState(InFileState)
+	: IChangelistTreeItem(bIsShelvedFile ? IChangelistTreeItem::ShelvedFile : IChangelistTreeItem::File)
+	, FileState(InFileState)
 	, MinTimeBetweenUpdate(FTimespan::FromSeconds(5.f))
 	, LastUpdateTime()
 	, bAssetsUpToDate(false)
 {
-	Type = (bIsShelvedFile ? IChangelistTreeItem::ShelvedFile : IChangelistTreeItem::File);
 	CheckBoxState = ECheckBoxState::Checked;
 
 	// Initialize asset data first
@@ -218,15 +214,22 @@ FString FFileTreeItem::RetrieveAssetPath(const FAssetData& InAssetData) const
 
 //////////////////////////////////////////////////////////////////////////
 
+FText FShelvedChangelistTreeItem::GetDisplayText() const
+{
+	return FText::Format(LOCTEXT("SourceControl_ShelvedFiles", "Shelved Items ({0})"), Children.Num());
+}
+
+//////////////////////////////////////////////////////////////////////////
+
 FOfflineFileTreeItem::FOfflineFileTreeItem(const FString& InFilename)
-	: Assets()
+	: IChangelistTreeItem(IChangelistTreeItem::OfflineFile)
+	, Assets()
 	, PackageName(FText::FromString(InFilename)) 
 	, AssetName(SSourceControlCommon::GetDefaultAssetName())
 	, AssetPath()
 	, AssetType(SSourceControlCommon::GetDefaultAssetType())
 	, AssetTypeColor()
 {
-	Type = IChangelistTreeItem::OfflineFile;
 	FString TempString;
 
 	USourceControlHelpers::GetAssetData(InFilename, Assets);
@@ -342,4 +345,4 @@ FText GetDefaultMultipleAsset()
 
 } // end of namespace SSourceControlCommon
 
-#undef LOCTEXT_NAMESPACE
\ No newline at end of file
+#undef LOCTEXT_NAMESPACE
diff --git a/Engine/Source/Editor/SourceControlWindows/Private/SSourceControlCommon.h b/Engine/Source/Editor/SourceControlWindows/Private/SSourceControlCommon.h
--- a/Engine/Source/Editor/SourceControlWindows/Private/SSourceControlCommon.h
+++ b/Engine/Source/Editor/SourceControlWindows/Private/SSourceControlCommon.h
@@ -10,57 +10,69 @@
 #include "UncontrolledChangelistState.h"
 
 struct FAssetData;
-struct IChangelistTreeItem;
-typedef TSharedPtr<IChangelistTreeItem> FChangelistTreeItemPtr;
-typedef TSharedRef<IChangelistTreeItem> FChangelistTreeItemRef;
 
+/**
+ * Modelizes a changelist node in a source control tree-like structure.
+ * The modelized tree stored is as below in memory.
+ * 
+ * > Changelist
+ *     File
+ *     > ShelvedChangelist
+ *         ShelvedFile
+ * 
+ * > UncontrolledChangelist
+ *     File
+ *     Offline File
+ */
 struct IChangelistTreeItem : TSharedFromThis<IChangelistTreeItem>
 {
 	enum TreeItemType
 	{
-		Invalid,
-		Changelist,
-		UncontrolledChangelist,
-		File,
-		OfflineFile,
-		ShelvedChangelist, // container for shelved files
-		ShelvedFile
+		Changelist,              // Node displaying a change list description.
+		UncontrolledChangelist,  // Node displaying an uncontrolled change list description.
+		File,                    // Node displaying a file information.
+		ShelvedChangelist,       // Node displaying shelved files as children.
+		ShelvedFile,             // Node displaying a shelved file information.
+		OfflineFile,             // Node displaying an offline file information.
 	};
 
-	/** Get this item's parent. Can be nullptr. */
-	FChangelistTreeItemPtr GetParent() const;
+	virtual ~IChangelistTreeItem() = default;
+
+	/** Get this item's parent. Can be nullptr for root nodes. */
+	TSharedPtr<IChangelistTreeItem> GetParent() const;
 
 	/** Get this item's children, if any. Although we store as weak pointers, they are guaranteed to be valid. */
-	const TArray<FChangelistTreeItemPtr>& GetChildren() const;
+	const TArray<TSharedPtr<IChangelistTreeItem>>& GetChildren() const;
 
 	/** Returns the TreeItem's type */
-	const TreeItemType GetTreeItemType() const
-	{
-		return Type;
-	}
+	const TreeItemType GetTreeItemType() const { return Type; }
 
 	/** Add a child to this item */
-	void AddChild(FChangelistTreeItemRef Child);
+	void AddChild(TSharedRef<IChangelistTreeItem> Child);
 
 	/** Remove a child from this item */
-	void RemoveChild(const FChangelistTreeItemRef& Child);
+	void RemoveChild(const TSharedRef<IChangelistTreeItem>& Child);
 
 protected:
+	IChangelistTreeItem(TreeItemType InType) { Type = InType; }
+
 	/** This item's parent, if any. */
-	FChangelistTreeItemPtr Parent;
+	TSharedPtr<IChangelistTreeItem> Parent;
 
 	/** Array of children contained underneath this item */
-	TArray<FChangelistTreeItemPtr> Children;
+	TArray<TSharedPtr<IChangelistTreeItem>> Children;
 
+	/** This item type. */
 	TreeItemType Type;
 };
 
+/** Displays a changelist icon/number/description. */
 struct FChangelistTreeItem : public IChangelistTreeItem
 {
-	FChangelistTreeItem(FSourceControlChangelistStateRef InChangelistState)
-		: ChangelistState(InChangelistState)
+	FChangelistTreeItem(TSharedRef<ISourceControlChangelistState> InChangelistState)
+		: IChangelistTreeItem(IChangelistTreeItem::Changelist)
+		, ChangelistState(MoveTemp(InChangelistState))
 	{
-		Type = IChangelistTreeItem::Changelist;
 	}
 
 	FText GetDisplayText() const
@@ -73,15 +85,26 @@ struct FChangelistTreeItem : public IChangelistTreeItem
 		return ChangelistState->GetDescriptionText();
 	}
 
-	FSourceControlChangelistStateRef ChangelistState;
+	int32 GetFileCount() const
+	{
+		return ChangelistState->GetFilesStates().Num();
+	}
+
+	int32 GetShelvedFileCount() const
+	{
+		return ChangelistState->GetShelvedFilesStates().Num();
+	}
+
+	TSharedRef<ISourceControlChangelistState> ChangelistState;
 };
 
+/** Displays an uncontrolled changelist icon/number/description. */
 struct FUncontrolledChangelistTreeItem : public IChangelistTreeItem
 {
 	FUncontrolledChangelistTreeItem(FUncontrolledChangelistStateRef InUncontrolledChangelistState)
-		: UncontrolledChangelistState(InUncontrolledChangelistState)
+		: IChangelistTreeItem(IChangelistTreeItem::UncontrolledChangelist)
+		, UncontrolledChangelistState(InUncontrolledChangelistState)
 	{
-		Type = IChangelistTreeItem::UncontrolledChangelist;
 	}
 
 	FText GetDisplayText() const
@@ -94,22 +117,20 @@ struct FUncontrolledChangelistTreeItem : public IChangelistTreeItem
 		return UncontrolledChangelistState->GetDescriptionText();
 	}
 
-	FUncontrolledChangelistStateRef UncontrolledChangelistState;
-};
-
-typedef TSharedPtr<FUncontrolledChangelistTreeItem> FUncontrolledChangelistTreeItemPtr;
-typedef TSharedRef<FUncontrolledChangelistTreeItem> FUncontrolledChangelistTreeItemRef;
+	int32 GetFileCount() const
+	{
+		return UncontrolledChangelistState->GetFilesStates().Num();
+	}
 
-struct FShelvedChangelistTreeItem : public IChangelistTreeItem
-{
-	FShelvedChangelistTreeItem()
+	int32 GetOfflineFileCount() const
 	{
-		Type = IChangelistTreeItem::ShelvedChangelist;
+		return UncontrolledChangelistState->GetOfflineFiles().Num();
 	}
 
-	FText GetDisplayText() const;
+	FUncontrolledChangelistStateRef UncontrolledChangelistState;
 };
 
+/** Displays a set of files under a changelist or uncontrolled changelist. */
 struct FFileTreeItem : public IChangelistTreeItem
 {
 	explicit FFileTreeItem(FSourceControlStateRef InFileState, bool bBeautifyPaths = true, bool bIsShelvedFile = false);
@@ -209,6 +230,21 @@ private:
 	bool bAssetsUpToDate;
 };
 
+/** Root node to group shelved files as children. */
+struct FShelvedChangelistTreeItem : public IChangelistTreeItem
+{
+	FShelvedChangelistTreeItem() : IChangelistTreeItem(IChangelistTreeItem::ShelvedChangelist) {}
+	FText GetDisplayText() const;
+};
+
+struct FShelvedFileTreeItem : public FFileTreeItem
+{
+	explicit FShelvedFileTreeItem(FSourceControlStateRef InFileState, bool bBeautifyPaths = true)
+		: FFileTreeItem(InFileState, bBeautifyPaths,/*bIsShelved=*/true)
+	{
+	}
+};
+
 struct FOfflineFileTreeItem : public IChangelistTreeItem
 {
 	explicit FOfflineFileTreeItem(const FString& InFilename);
@@ -233,16 +269,6 @@ private:
 	FSlateColor AssetTypeColor;
 };
 
-typedef TSharedPtr<FFileTreeItem> FFileTreeItemPtr;
-typedef TSharedRef<FFileTreeItem> FFileTreeItemRef;
-
-struct FShelvedFileTreeItem : public FFileTreeItem
-{
-	explicit FShelvedFileTreeItem(FSourceControlStateRef InFileState, bool bBeautifyPaths = true)
-		: FFileTreeItem(InFileState, bBeautifyPaths,/*bIsShelved=*/true)
-	{}
-};
-
 namespace SSourceControlCommon
 {
 	TSharedRef<SWidget> GetSCCFileWidget(FSourceControlStateRef InFileState, bool bIsShelvedFile = false);
@@ -250,4 +276,13 @@ namespace SSourceControlCommon
 	FText GetDefaultAssetType();
 	FText GetDefaultUnknownAssetType();
 	FText GetDefaultMultipleAsset();
-}
\ No newline at end of file
+}
+
+
+typedef TSharedPtr<FUncontrolledChangelistTreeItem> FUncontrolledChangelistTreeItemPtr;
+typedef TSharedRef<FUncontrolledChangelistTreeItem> FUncontrolledChangelistTreeItemRef;
+typedef TSharedPtr<IChangelistTreeItem> FChangelistTreeItemPtr;
+typedef TSharedRef<IChangelistTreeItem> FChangelistTreeItemRef;
+typedef TSharedPtr<FFileTreeItem> FFileTreeItemPtr;
+typedef TSharedRef<FFileTreeItem> FFileTreeItemRef;
+
diff --git a/Engine/Source/Editor/SourceControlWindows/SourceControlWindows.Build.cs b/Engine/Source/Editor/SourceControlWindows/SourceControlWindows.Build.cs
--- a/Engine/Source/Editor/SourceControlWindows/SourceControlWindows.Build.cs
+++ b/Engine/Source/Editor/SourceControlWindows/SourceControlWindows.Build.cs
@@ -16,6 +16,7 @@ public class SourceControlWindows : ModuleRules
 				"CoreUObject", 
                 "InputCore",
 				"Engine", 
+				"OutputLog",
 				"Slate",
 				"SlateCore",
 				"SourceControl", 
@@ -23,7 +24,7 @@ public class SourceControlWindows : ModuleRules
 				"AssetTools",
 				"EditorFramework",
 				"WorkspaceMenuStructure",
-				"UnrealEd"			// We need this dependency here because we use PackageTools.
+				"UnrealEd",		// We need this dependency here because we use PackageTools.
 			}
 		);
 
