Fixed the source control history window not properly displaying multilines submit descriptions.

diff --git a/Engine/Source/Editor/SourceControlWindows/Private/SSourceControlHistory.cpp b/Engine/Source/Editor/SourceControlWindows/Private/SSourceControlHistory.cpp
--- a/Engine/Source/Editor/SourceControlWindows/Private/SSourceControlHistory.cpp
+++ b/Engine/Source/Editor/SourceControlWindows/Private/SSourceControlHistory.cpp
@@ -36,6 +36,7 @@
 #include "Widgets/Layout/SBox.h"
 #include "Framework/MultiBox/MultiBoxBuilder.h"
 #include "Widgets/Input/SEditableTextBox.h"
+#include "Widgets/Input/SMultiLineEditableTextBox.h"
 #include "Widgets/Layout/SSplitter.h"
 #include "Widgets/Views/SExpanderArrow.h"
 #include "Widgets/Views/SHeaderRow.h"
@@ -923,12 +924,12 @@ private:
 				[
 					SNew(SVerticalBox)
 					+SVerticalBox::Slot()
-					.AutoHeight()
 					.Padding(5)
 					[
-						SNew(SEditableText)
+						SNew(SMultiLineEditableTextBox)
 						.IsReadOnly(true)
-						.Text(this, &SSourceControlHistoryWidget::GetDescription)	
+						.AutoWrapText(true)
+						.Text(this, &SSourceControlHistoryWidget::GetDescription)
 					]
 				]
 			]
