Fixed circular dependecies between the SourceControlWindows and OutputModule reported on Linux.

diff --git a/Engine/Source/Editor/SourceControlWindows/Private/SSourceControlChangelists.cpp b/Engine/Source/Editor/SourceControlWindows/Private/SSourceControlChangelists.cpp
--- a/Engine/Source/Editor/SourceControlWindows/Private/SSourceControlChangelists.cpp
+++ b/Engine/Source/Editor/SourceControlWindows/Private/SSourceControlChangelists.cpp
@@ -38,7 +38,6 @@
 #include "Framework/Application/SlateApplication.h"
 #include "Framework/Notifications/NotificationManager.h"
 #include "Framework/Docking/TabManager.h"
-#include "OutputLogModule.h"
 
 #define LOCTEXT_NAMESPACE "SourceControlChangelist"
 
@@ -69,8 +68,7 @@ void DisplaySourceControlOperationNotification(const FText& Message, SNotificati
 
 	FNotificationInfo NotificationInfo(Message);
 	NotificationInfo.ExpireDuration = 6.0f;
-	//NotificationInfo.Hyperlink = FSimpleDelegate::CreateLambda([]() { FOutputLogModule::Get().FocusOutputLog(); }); // NOTE: This is what we want moving further, but this is only in main branch for the moment.
-	NotificationInfo.Hyperlink = FSimpleDelegate::CreateLambda([]() { FGlobalTabmanager::Get()->TryInvokeTab(FName("OutputLog")); }); // Replace this once with the above once this change reach the main branch.
+	NotificationInfo.Hyperlink = FSimpleDelegate::CreateLambda([]() { FGlobalTabmanager::Get()->TryInvokeTab(FName("OutputLog")); });
 	NotificationInfo.HyperlinkText = LOCTEXT("ShowOutputLogHyperlink", "Show Output Log");
 	FSlateNotificationManager::Get().AddNotification(NotificationInfo)->SetCompletionState(CompletionState);
 }
diff --git a/Engine/Source/Editor/SourceControlWindows/SourceControlWindows.Build.cs b/Engine/Source/Editor/SourceControlWindows/SourceControlWindows.Build.cs
--- a/Engine/Source/Editor/SourceControlWindows/SourceControlWindows.Build.cs
+++ b/Engine/Source/Editor/SourceControlWindows/SourceControlWindows.Build.cs
@@ -16,7 +16,6 @@ public class SourceControlWindows : ModuleRules
 				"CoreUObject", 
                 "InputCore",
 				"Engine", 
-				"OutputLog",
 				"Slate",
 				"SlateCore",
 				"SourceControl", 
