Modified the source control submit window opened from the changelist window:
  - Supported editing the changelist and saving it in P4. If the submit fails, the code will also try to save any edited changelist description.
  - Enabled the 'keep checked out' button when submitting from the changelist window.
  - Fixed the P4 cache not properly updating the file 'kept in checkout'.
  - If the submit window is opened from the Source Control menu at the bottom right of the Eitor, the new 'Update' button to save the changelist description is not visible, this is not supported for all source control provider.

Added a delay of 1.5 seconds before showing the slow stack in the changelist window when refreshing the UI.

diff --git a/Engine/Plugins/Developer/PerforceSourceControl/Source/PerforceSourceControl/Private/PerforceSourceControlOperations.cpp b/Engine/Plugins/Developer/PerforceSourceControl/Source/PerforceSourceControl/Private/PerforceSourceControlOperations.cpp
--- a/Engine/Plugins/Developer/PerforceSourceControl/Source/PerforceSourceControl/Private/PerforceSourceControlOperations.cpp
+++ b/Engine/Plugins/Developer/PerforceSourceControl/Source/PerforceSourceControl/Private/PerforceSourceControlOperations.cpp
@@ -614,19 +614,24 @@ bool FPerforceCheckInWorker::Execute(FPerforceSourceControlCommand& InCommand)
 				// Remove any deleted files from status cache
 				TArray<TSharedRef<ISourceControlState, ESPMode::ThreadSafe>> States;
 				GetSCCProvider().GetState(FilesToSubmit, States, EStateCacheUsage::Use);
-				for (const auto& State : States)
+				for (const TSharedRef<ISourceControlState>& State : States)
 				{
 					if (State->IsDeleted())
 					{
 						GetSCCProvider().RemoveFileFromCache(State->GetFilename());
 					}
+
+					if (Operation->GetKeepCheckedOut() && (State->IsAdded() || State->IsCheckedOut()))
+					{
+						FilesToKeepCheckedOut.Emplace(State->GetFilename(), EPerforceState::CheckedOut);
+					}
 				}
 
 				StaticCastSharedRef<FCheckIn>(InCommand.Operation)->SetSuccessMessage(ParseSubmitResults(Records));
 
-				for(auto Iter(FilesToSubmit.CreateIterator()); Iter; Iter++)
+				for (const FString& Filename : FilesToSubmit)
 				{
-					OutResults.Add(*Iter, EPerforceState::ReadOnly);
+					OutResults.Add(Filename, EPerforceState::ReadOnly);
 				}
 
 				InChangelist = InCommand.Changelist;
@@ -682,6 +687,14 @@ bool FPerforceCheckInWorker::UpdateStates() const
 		{
 			bUpdatedChangelistStates = RemoveFilesFromChangelist(OutResults, GetSCCProvider(), InChangelist);
 		}
+
+		// Ensure to update the cache with respect to the file kept in checkout. Those files are move to the 'default' changelist by P4.
+		if (!FilesToKeepCheckedOut.IsEmpty())
+		{
+			UpdateChangelistState(GetSCCProvider(), FPerforceSourceControlChangelist::DefaultChangelist, FilesToKeepCheckedOut);
+			UpdateCachedStates(GetSCCProvider(), FilesToKeepCheckedOut);
+			bUpdatedChangelistStates = true;
+		}
 	}
 
 	return (bUpdatedStates || bUpdatedChangelistStates);
diff --git a/Engine/Plugins/Developer/PerforceSourceControl/Source/PerforceSourceControl/Private/PerforceSourceControlOperations.h b/Engine/Plugins/Developer/PerforceSourceControl/Source/PerforceSourceControl/Private/PerforceSourceControlOperations.h
--- a/Engine/Plugins/Developer/PerforceSourceControl/Source/PerforceSourceControl/Private/PerforceSourceControlOperations.h
+++ b/Engine/Plugins/Developer/PerforceSourceControl/Source/PerforceSourceControl/Private/PerforceSourceControlOperations.h
@@ -67,6 +67,9 @@ public:
 
 	/** Changelist we submitted */
 	FPerforceSourceControlChangelist OutChangelist;
+
+	/** List of file to kept in checked out after the submit.*/
+	TMap<FString, EPerforceState::Type> FilesToKeepCheckedOut;
 };
 
 class FPerforceGetFileListWorker final : public IPerforceSourceControlWorker
diff --git a/Engine/Source/Editor/SourceControlWindows/Private/SSourceControlChangelists.cpp b/Engine/Source/Editor/SourceControlWindows/Private/SSourceControlChangelists.cpp
--- a/Engine/Source/Editor/SourceControlWindows/Private/SSourceControlChangelists.cpp
+++ b/Engine/Source/Editor/SourceControlWindows/Private/SSourceControlChangelists.cpp
@@ -41,8 +41,8 @@
 
 #define LOCTEXT_NAMESPACE "SourceControlChangelist"
 
-
-const FText SSourceControlChangelistsWidget::ChangelistValidatedTag = LOCTEXT("ValidationTag", "#changelist validated");
+namespace
+{
 
 /** Wraps the execution of a changelist operations with a slow task. */
 void ExecuteChangelistOperationWithSlowTaskWrapper(const FText& Message, const TFunction<void()>& ChangelistTask)
@@ -85,12 +85,49 @@ bool AreUncontrolledChangelistsEnabled()
 	return FUncontrolledChangelistsModule::Get().IsEnabled();
 };
 
-/** Returns true if they are changelists to display. */
+/** Returns true if there are changelists to display. */
 bool AreChangelistsEnabled()
 {
 	return AreControlledChangelistsEnabled() || AreUncontrolledChangelistsEnabled();
 };
 
+/**
+ * Returns a new changelist description if needed, appending validation tag.
+ * 
+ * @param bInValidationResult	The result of the validation step
+ * @param InOriginalChangelistDescription	Description of the changelist before modification
+ * 
+ * @return The new changelist description
+ */
+FText UpdateChangelistDescriptionToSubmitIfNeeded(const bool bInValidationResult, const FText& InChangelistDescription)
+{
+	auto GetChangelistValidationTag = []
+	{
+		return LOCTEXT("ValidationTag", "#changelist validated");
+	};
+
+	auto ContainsValidationFlag = [&GetChangelistValidationTag](const FText& InChangelistDescription)
+	{
+		FString DescriptionString = InChangelistDescription.ToString();
+		FString ValidationString = GetChangelistValidationTag().ToString();
+		return DescriptionString.Find(ValidationString) != INDEX_NONE;
+	};
+
+	if (bInValidationResult && USourceControlPreferences::IsValidationTagEnabled() && !ContainsValidationFlag(InChangelistDescription))
+	{
+		FStringOutputDevice Str;
+
+		Str.SetAutoEmitLineTerminator(true);
+		Str.Log(InChangelistDescription);
+		Str.Log(GetChangelistValidationTag());
+
+		return FText::FromString(Str);
+	}
+
+	return InChangelistDescription;
+}
+
+} // Anonymous namespace
 
 /** Implements drag and drop operation. */
 struct FSCCFileDragDropOp : public FDragDropOperation
@@ -474,14 +511,6 @@ TSharedRef<SWidget> SSourceControlChangelistsWidget::MakeToolBar()
 	return ToolBarBuilder.MakeWidget();
 }
 
-bool SSourceControlChangelistsWidget::HasValidationTag(const FText& InChangelistDescription) const
-{
-	FString DescriptionString = InChangelistDescription.ToString();
-	FString ValidationString = ChangelistValidatedTag.ToString();
-
-	return DescriptionString.Find(ValidationString) != INDEX_NONE;
-}
-
 void SSourceControlChangelistsWidget::EditChangelistDescription(const FText& InNewChangelistDescription, const FSourceControlChangelistStatePtr& InChangelistState)
 {
 	TSharedRef<FEditChangelist> EditChangelistOperation = ISourceControlOperation::Create<FEditChangelist>();
@@ -574,7 +603,7 @@ void SSourceControlChangelistsWidget::StartRefreshStatus()
 
 void SSourceControlChangelistsWidget::TickRefreshStatus(double InDeltaTime)
 {
-	int RefreshStatusTimeElapsed = static_cast<int>(FPlatformTime::Seconds() - RefreshStatusStartSecs);
+	int32 RefreshStatusTimeElapsed = static_cast<int32>(FPlatformTime::Seconds() - RefreshStatusStartSecs);
 	RefreshStatus = FText::Format(LOCTEXT("SourceControl_RefreshStatus", "Refreshing changelists... ({0} s)"), FText::AsNumber(RefreshStatusTimeElapsed));
 }
 
@@ -638,7 +667,7 @@ void SSourceControlChangelistsWidget::OnRefresh()
 	}
 
 	FScopedSlowTask SlowTask(ElementsToProcess, LOCTEXT("SourceControl_RebuildTree", "Refreshing Tree Items"));
-	SlowTask.MakeDialog(/*bShowCancelButton=*/true);
+	SlowTask.MakeDialogDelayed(1.5f, /*bShowCancelButton=*/true);
 
 	// Rebuild the tree data models
 	bool bBeautifyPaths = true;
@@ -1391,10 +1420,15 @@ void SSourceControlChangelistsWidget::OnSubmitChangelist()
 	FString ChangelistValidationErrorsText;
 	bool bValidationResult = GetChangelistValidationResult(ChangelistState->GetChangelist(), ChangelistValidationTitle, ChangelistValidationWarningsText, ChangelistValidationErrorsText);
 
-	// Build list of states for the dialog
-	const FText OriginalChangelistDescription = ChangelistState->GetDescriptionText();
-	const bool bAskForChangelistDescription = (OriginalChangelistDescription.IsEmptyOrWhitespace());
-	FText ChangelistDescriptionToSubmit = UpdateChangelistDescriptionToSubmitIfNeeded(bValidationResult, OriginalChangelistDescription);
+	// The description from the source control.
+	const FText CurrentChangelistDescription = ChangelistState->GetDescriptionText();
+	const bool bAskForChangelistDescription = (CurrentChangelistDescription.IsEmptyOrWhitespace());
+
+	// The description possibly updated with the #validated proposed to the user.
+	FText ChangelistDescriptionToSubmit = UpdateChangelistDescriptionToSubmitIfNeeded(bValidationResult, CurrentChangelistDescription);
+
+	// The description once edited by the user in the Submit window.
+	FText UserEditChangelistDescription = ChangelistDescriptionToSubmit;
 
 	TSharedRef<SWindow> NewWindow = SNew(SWindow)
 		.Title(NSLOCTEXT("SourceControl.ConfirmSubmit", "Title", "Confirm changelist submit"))
@@ -1411,10 +1445,19 @@ void SSourceControlChangelistsWidget::OnSubmitChangelist()
 		.ChangeValidationResult(ChangelistValidationTitle)
 		.ChangeValidationWarnings(ChangelistValidationWarningsText)
 		.ChangeValidationErrors(ChangelistValidationErrorsText)
-		.AllowDescriptionChange(bAskForChangelistDescription)
+		.AllowDescriptionChange(true)
 		.AllowUncheckFiles(false)
-		.AllowKeepCheckedOut(false)
-		.AllowSubmit(bValidationResult);
+		.AllowKeepCheckedOut(true)
+		.AllowSubmit(bValidationResult)
+		.OnSaveChangelistDescription(
+			FSourceControlSaveChangelistDescription::CreateLambda([this, &ChangelistState, &bValidationResult, &UserEditChangelistDescription](const FText& NewDescription)
+			{
+				// NOTE this is called from a modal dialog, so adding a slow task on top of it doesn't really look good. Just run a synchronous operation.
+				TSharedRef<FEditChangelist> EditChangelistOperation = ISourceControlOperation::Create<FEditChangelist>();
+				EditChangelistOperation->SetDescription(NewDescription);
+				ISourceControlModule::Get().GetProvider().Execute(EditChangelistOperation, ChangelistState->GetChangelist(), EConcurrency::Synchronous);
+				UserEditChangelistDescription = NewDescription;
+			}));
 
 	NewWindow->SetContent(
 		SourceControlWidget
@@ -1426,61 +1469,48 @@ void SSourceControlChangelistsWidget::OnSubmitChangelist()
 	{
 		ISourceControlProvider& SourceControlProvider = ISourceControlModule::Get().GetProvider();
 		FChangeListDescription Description;
-		auto SubmitChangelistOperation = ISourceControlOperation::Create<FCheckIn>();
+		TSharedRef<FCheckIn> SubmitChangelistOperation = ISourceControlOperation::Create<FCheckIn>();
+		SubmitChangelistOperation->SetKeepCheckedOut(SourceControlWidget->WantToKeepCheckedOut());
 		bool bCheckinSuccess = false;
 
+		// Get the changelist description the user had when he hit the 'submit' button.
 		SourceControlWidget->FillChangeListDescription(Description);
+		UserEditChangelistDescription = Description.Description;
 
-		// Check if any of the presubmit hooks fail and if so early out to avoid the submit
-		if (!GetOnPresubmitResult(ChangelistState, Description))
+		// Check if any of the presubmit hooks fail. (This might also update the changelist description)
+		if (GetOnPresubmitResult(ChangelistState, Description))
 		{
-			return;
-		}
+			// If the description was modified, add it to the operation to update the changelist
+			if (!ChangelistDescriptionToSubmit.EqualTo(Description.Description))
+			{
+				SubmitChangelistOperation->SetDescription(UpdateChangelistDescriptionToSubmitIfNeeded(bValidationResult, Description.Description));
+			}
 
-		// If the description was modified, we add it to the operation to update the changelist
-		if (!OriginalChangelistDescription.EqualTo(Description.Description))
-		{
-			SubmitChangelistOperation->SetDescription(Description.Description);
+			Execute(LOCTEXT("Submitting_Changelist", "Submitting changelist..."), SubmitChangelistOperation, ChangelistState->GetChangelist(), EConcurrency::Synchronous, FSourceControlOperationComplete::CreateLambda(
+				[&SubmitChangelistOperation, &bCheckinSuccess](const TSharedRef<ISourceControlOperation>& Operation, ECommandResult::Type InResult)
+				{
+					if (InResult == ECommandResult::Succeeded)
+					{
+						DisplaySourceControlOperationNotification(SubmitChangelistOperation->GetSuccessMessage(), SNotificationItem::CS_Success);
+						bCheckinSuccess = true;
+					}
+					else if (InResult == ECommandResult::Failed)
+					{
+						DisplaySourceControlOperationNotification(LOCTEXT("SCC_Checkin_Failed", "Failed to check in files!"), SNotificationItem::CS_Fail);
+					}
+				}));
 		}
 
-		bCheckinSuccess = SourceControlProvider.Execute(SubmitChangelistOperation, ChangelistState->GetChangelist()) == ECommandResult::Succeeded;
-
-		// Setup the notification for operation feedback
-		FNotificationInfo Info(SubmitChangelistOperation->GetSuccessMessage());
-
-		// Override the notification fields for failure ones
-		if (!bCheckinSuccess)
+		// If something went wrong with the submit, try to preserve the changelist edited by the user (if he edited).
+		if (!bCheckinSuccess && !UserEditChangelistDescription.EqualTo(ChangelistDescriptionToSubmit))
 		{
-			Info.Text = LOCTEXT("SCC_Checkin_Failed", "Failed to check in files!");
+			TSharedRef<FEditChangelist> EditChangelistOperation = ISourceControlOperation::Create<FEditChangelist>();
+			EditChangelistOperation->SetDescription(UserEditChangelistDescription);
+			ISourceControlModule::Get().GetProvider().Execute(EditChangelistOperation, ChangelistState->GetChangelist(), EConcurrency::Synchronous);
 		}
-		
-		Info.ExpireDuration = 8.0f;
-		Info.HyperlinkText = LOCTEXT("SCC_Checkin_ShowLog", "Show Message Log");
-		Info.Hyperlink = FSimpleDelegate::CreateLambda([]() { FMessageLog("SourceControl").Open(EMessageSeverity::Info, true); });
-
-		TSharedPtr<SNotificationItem> Notification = FSlateNotificationManager::Get().AddNotification(Info);
-		Notification->SetCompletionState(bCheckinSuccess ? SNotificationItem::CS_Success : SNotificationItem::CS_Fail);
 	}
 }
 
-FText SSourceControlChangelistsWidget::UpdateChangelistDescriptionToSubmitIfNeeded(const bool bInValidationResult, const FText& InOriginalChangelistDescription) const
-{
-	if (bInValidationResult && USourceControlPreferences::IsValidationTagEnabled() && (!HasValidationTag(InOriginalChangelistDescription)))
-	{
-		FStringOutputDevice Str;
-
-		Str.SetAutoEmitLineTerminator(true);
-		Str.Log(InOriginalChangelistDescription);
-		Str.Log(ChangelistValidatedTag);
-
-		FText ChangelistDescription = FText::FromString(Str);
-
-		return ChangelistDescription;
-	}
-
-	return InOriginalChangelistDescription;
-}
-
 bool SSourceControlChangelistsWidget::CanSubmitChangelist()
 {
 	FSourceControlChangelistStatePtr Changelist = GetCurrentChangelistState();
diff --git a/Engine/Source/Editor/SourceControlWindows/Private/SSourceControlChangelists.h b/Engine/Source/Editor/SourceControlWindows/Private/SSourceControlChangelists.h
--- a/Engine/Source/Editor/SourceControlWindows/Private/SSourceControlChangelists.h
+++ b/Engine/Source/Editor/SourceControlWindows/Private/SSourceControlChangelists.h
@@ -141,26 +141,11 @@ private:
 	void RestoreExpandedAndSelectionStates(const FExpandedAndSelectionStates& InStates);
 
 	TSharedRef<SWidget> MakeToolBar();
-
-	/**
-	 * Returns a new changelist description if needed, appending validation tag.
-	 * 
-	 * @param bInValidationResult	The result of the validation step
-	 * @param InOriginalChangelistDescription	Description of the changelist before modification
-	 * 
-	 * @return The new changelist description
-	 */
-	FText UpdateChangelistDescriptionToSubmitIfNeeded(const bool bInValidationResult, const FText& InOriginalChangelistDescription) const;
 	
-	/** Returns true if the provided changelist description contains a validation tag. */
-	bool HasValidationTag(const FText& InChangelistDescription) const;
-
 	/** Executes an operation to updates the changelist description of the provided changelist with a new description. */
 	void EditChangelistDescription(const FText& InNewChangelistDescription, const FSourceControlChangelistStatePtr& InChangelistState);
 
 private:
-	/** Tag to append to a changelist that passed validation */
-	static const FText ChangelistValidatedTag;
 
 	TSharedPtr<SExpandableChangelistArea> ChangelistExpandableArea;
 	TSharedPtr<SExpandableChangelistArea> UncontrolledChangelistExpandableArea;
@@ -182,9 +167,6 @@ private:
 	/** Source control state changed delegate handle */
 	FDelegateHandle SourceControlStateChangedDelegateHandle;
 
-	/** Uncontrolled Changelist changed delegate handle */
-	FDelegateHandle UncontrolledChangelistChangedDelegateHandle;
-
 	bool bShouldRefresh = false;
 	bool bSourceControlAvailable = false;
 
diff --git a/Engine/Source/Editor/SourceControlWindows/Private/SSourceControlSubmit.cpp b/Engine/Source/Editor/SourceControlWindows/Private/SSourceControlSubmit.cpp
--- a/Engine/Source/Editor/SourceControlWindows/Private/SSourceControlSubmit.cpp
+++ b/Engine/Source/Editor/SourceControlWindows/Private/SSourceControlSubmit.cpp
@@ -31,6 +31,10 @@
 
 #define LOCTEXT_NAMESPACE "SSourceControlSubmit"
 
+// This is useful for source control that do not support changelist (Git/SVN) or when the submit widget is not created from the changelist window. If a user
+// commits/submits this way, then edits the submit description but cancels, the description will be remembered in memory for the next time he tries to submit.
+static FText GSavedChangeListDescription;
+
 bool TryToVirtualizeFilesToSubmit(const TArray<FString>& FilesToSubmit, FText& Description, FText& OutFailureMsg)
 {
 	// TODO: Once this is removed and not deprecated move the following arrays inside of the System.IsEnabled() scope
@@ -119,12 +123,11 @@ TSharedRef<SWidget> SSourceControlSubmitListRow::GenerateWidgetForColumn(const F
 	return SNullWidget::NullWidget;
 }
 
-FText SSourceControlSubmitWidget::SavedChangeListDescription;
 
 SSourceControlSubmitWidget::~SSourceControlSubmitWidget()
 {
-	// If the user cancel the submit, save the changelist. (On successful submit, ChangeListDescriptionTextCtrl is cleared).
-	SavedChangeListDescription = ChangeListDescriptionTextCtrl->GetText();
+	// If the user cancel the submit, save the changelist. If the user submitted, ChangeListDescriptionTextCtrl was cleared).
+	GSavedChangeListDescription = ChangeListDescriptionTextCtrl->GetText();
 }
 
 void SSourceControlSubmitWidget::Construct(const FArguments& InArgs)
@@ -134,8 +137,8 @@ void SSourceControlSubmitWidget::Construct(const FArguments& InArgs)
 	SortMode = EColumnSortMode::Ascending;
 	if (!InArgs._Description.Get().IsEmpty())
 	{
-		// This is meant to restore the description entered by the user if they canceled a previous submit operation (or the submit failed somehow).
-		SavedChangeListDescription = InArgs._Description.Get();
+		// If a description is provided, override the last one saved in memory.
+		GSavedChangeListDescription = InArgs._Description.Get();
 	}
 	bAllowSubmit = InArgs._AllowSubmit.Get();
 
@@ -143,6 +146,7 @@ void SSourceControlSubmitWidget::Construct(const FArguments& InArgs)
 	const bool bAllowUncheckFiles = InArgs._AllowUncheckFiles.Get();
 	const bool bAllowKeepCheckedOut = InArgs._AllowKeepCheckedOut.Get();
 	const bool bShowChangelistValidation = !InArgs._ChangeValidationResult.Get().IsEmpty();
+	OnSaveChangelistDescription = InArgs._OnSaveChangelistDescription;
 
 	for (const auto& Item : InArgs._Items.Get())
 	{
@@ -219,7 +223,7 @@ void SSourceControlSubmitWidget::Construct(const FArguments& InArgs)
 		[
 			SAssignNew(ChangeListDescriptionTextCtrl, SMultiLineEditableTextBox)
 			.SelectAllTextWhenFocused(!bDescriptionIsReadOnly)
-			.Text(SavedChangeListDescription)
+			.Text(GSavedChangeListDescription)
 			.AutoWrapText(true)
 			.IsReadOnly(bDescriptionIsReadOnly)
 		]
@@ -252,7 +256,9 @@ void SSourceControlSubmitWidget::Construct(const FArguments& InArgs)
 			.Padding(5)
 			[
 				SNew( SErrorText )
-				.ErrorText( NSLOCTEXT("SourceControl.SubmitPanel", "ChangeListDescWarning", "Changelist description is required to submit") )
+				.ErrorText(ChangeListDescriptionTextCtrl->GetText().IsEmpty() ? 
+					NSLOCTEXT("SourceControl.SubmitPanel", "ChangeListDescWarning", "Changelist description is required to submit") :
+					NSLOCTEXT("SourceControl.SubmitPanel", "Error", "Error!")) // Other errors exist and a better mechanism should be built in to display the right error. 
 			]
 		];
 	}
@@ -378,6 +384,15 @@ void SSourceControlSubmitWidget::Construct(const FArguments& InArgs)
 		.MinDesiredSlotWidth(FAppStyle::GetFloat("StandardDialog.MinDesiredSlotWidth"))
 		.MinDesiredSlotHeight(FAppStyle::GetFloat("StandardDialog.MinDesiredSlotHeight"))
 		+SUniformGridPanel::Slot(0,0)
+		[
+			SNew(SButton)
+			.Visibility(OnSaveChangelistDescription.IsBound() ? EVisibility::Visible : EVisibility::Collapsed)
+			.HAlign(HAlign_Center)
+			.ContentPadding(FAppStyle::GetMargin("StandardDialog.ContentPadding"))
+			.Text( NSLOCTEXT("SourceControl.SubmitPanel", "Update", "Update") )
+			.OnClicked(FOnClicked::CreateLambda([this]() { OnSaveChangelistDescription.ExecuteIfBound(ChangeListDescriptionTextCtrl->GetText()); return FReply::Handled(); }))
+		]
+		+SUniformGridPanel::Slot(1,0)
 		[
 			SNew(SButton)
 			.HAlign(HAlign_Center)
@@ -386,7 +401,7 @@ void SSourceControlSubmitWidget::Construct(const FArguments& InArgs)
 			.Text( NSLOCTEXT("SourceControl.SubmitPanel", "OKButton", "Submit") )
 			.OnClicked(this, &SSourceControlSubmitWidget::SubmitClicked)
 		]
-		+SUniformGridPanel::Slot(1,0)
+		+SUniformGridPanel::Slot(2,0)
 		[
 			SNew(SButton)
 			.HAlign(HAlign_Center)
@@ -635,7 +650,7 @@ bool SSourceControlSubmitWidget::IsSubmitEnabled() const
 
 EVisibility SSourceControlSubmitWidget::IsWarningPanelVisible() const
 {
-	return IsSubmitEnabled()? EVisibility::Hidden : EVisibility::Visible;
+	return IsSubmitEnabled() ? EVisibility::Collapsed : EVisibility::Visible;
 }
 
 
@@ -742,7 +757,6 @@ void SSourceControlSubmitWidget::SortTree()
 	}
 }
 
-
 #undef LOCTEXT_NAMESPACE
 
 #endif // SOURCE_CONTROL_WITH_SLATE
diff --git a/Engine/Source/Editor/SourceControlWindows/Private/SSourceControlSubmit.h b/Engine/Source/Editor/SourceControlWindows/Private/SSourceControlSubmit.h
--- a/Engine/Source/Editor/SourceControlWindows/Private/SSourceControlSubmit.h
+++ b/Engine/Source/Editor/SourceControlWindows/Private/SSourceControlSubmit.h
@@ -42,6 +42,8 @@ public:
 
 bool TryToVirtualizeFilesToSubmit(const TArray<FString>& FilesToSubmit, FText& Description, FText& OutFailureMsg);
 
+DECLARE_DELEGATE_OneParam(FSourceControlSaveChangelistDescription, const FText& /*NewDescription*/);
+
 class SSourceControlSubmitWidget : public SCompoundWidget
 {
 public:
@@ -56,6 +58,7 @@ public:
 		, _AllowUncheckFiles(true)
 		, _AllowKeepCheckedOut(true)
 		, _AllowSubmit(true)
+		, _OnSaveChangelistDescription()
 	{}
 
 		SLATE_ATTRIBUTE(TSharedPtr<SWindow>, ParentWindow)
@@ -68,6 +71,7 @@ public:
 		SLATE_ATTRIBUTE(bool, AllowUncheckFiles)
 		SLATE_ATTRIBUTE(bool, AllowKeepCheckedOut)
 		SLATE_ATTRIBUTE(bool, AllowSubmit)
+		SLATE_EVENT(FSourceControlSaveChangelistDescription, OnSaveChangelistDescription)
 
 	SLATE_END_ARGS()
 
@@ -165,7 +169,7 @@ private:
 	bool CanDiffAgainstDepot() const;
 	void OnDiffAgainstDepot();
 	void OnDiffAgainstDepotSelected(TSharedPtr<FFileTreeItem> InSelectedItem);
-	
+
 private:
 	ESubmitResults::Type DialogResult;
 
@@ -193,8 +197,8 @@ private:
 	/** Currently selected sorting mode */
 	EColumnSortMode::Type SortMode;
 
-	/** Submit Description saved when the widget is destroyed if canceled */
-	static FText SavedChangeListDescription;	
+	/** Delegate invoked to save the description of a changelist with the source control. */
+	FSourceControlSaveChangelistDescription OnSaveChangelistDescription;
 };
 
 class SSourceControlSubmitListRow : public SMultiColumnTableRow<TSharedPtr<FFileTreeItem>>
