Fix ClearChangeListDescription() not getting called when submitted a Changelist from the changelist window.

PR #9432: [Source Control][UE5.1] Fix regression ClearChangeListDescription() not getting called when submitted a Changelist (Contributed by [REDACTED_PERSON])

diff --git a/Engine/Source/Editor/SourceControlWindows/Private/SSourceControlChangelists.cpp b/Engine/Source/Editor/SourceControlWindows/Private/SSourceControlChangelists.cpp
--- a/Engine/Source/Editor/SourceControlWindows/Private/SSourceControlChangelists.cpp
+++ b/Engine/Source/Editor/SourceControlWindows/Private/SSourceControlChangelists.cpp
@@ -1508,6 +1508,14 @@ void SSourceControlChangelistsWidget::OnSubmitChangelist()
 			EditChangelistOperation->SetDescription(UserEditChangelistDescription);
 			ISourceControlModule::Get().GetProvider().Execute(EditChangelistOperation, ChangelistState->GetChangelist(), EConcurrency::Synchronous);
 		}
+
+		if (bCheckinSuccess)
+		{
+			// Clear the description saved by the 'submit window'. Useful when the submit window is opened from the Editor menu rather than the changelist window.
+			// Opening the 'submit window' from the Editor menu is intended for source controls that do not support changelists (SVN/Git), but remains available to
+			// all source controls at the moment.
+			SourceControlWidget->ClearChangeListDescription();
+		}
 	}
 }
 
