Renamed the 'Update' button as 'Apply' in the submit changelist window to mimic P4 vocabulary and added a tooltip.

diff --git a/Engine/Source/Editor/SourceControlWindows/Private/SSourceControlSubmit.cpp b/Engine/Source/Editor/SourceControlWindows/Private/SSourceControlSubmit.cpp
--- a/Engine/Source/Editor/SourceControlWindows/Private/SSourceControlSubmit.cpp
+++ b/Engine/Source/Editor/SourceControlWindows/Private/SSourceControlSubmit.cpp
@@ -389,7 +389,8 @@ void SSourceControlSubmitWidget::Construct(const FArguments& InArgs)
 			.Visibility(OnSaveChangelistDescription.IsBound() ? EVisibility::Visible : EVisibility::Collapsed)
 			.HAlign(HAlign_Center)
 			.ContentPadding(FAppStyle::GetMargin("StandardDialog.ContentPadding"))
-			.Text( NSLOCTEXT("SourceControl.SubmitPanel", "Update", "Update") )
+			.Text(NSLOCTEXT("SourceControl.SubmitPanel", "Apply", "Apply"))
+			.ToolTipText(NSLOCTEXT("SourceControl.SubmitPanel", "Apply_Tooltip", "Save the description without submitting."))
 			.OnClicked(FOnClicked::CreateLambda([this]() { OnSaveChangelistDescription.ExecuteIfBound(ChangeListDescriptionTextCtrl->GetText()); return FReply::Handled(); }))
 		]
 		+SUniformGridPanel::Slot(1,0)
