SceneOutliner:

- SCC Show in changelist
- Fix Unloaded actors not having SCC Menu

diff --git a/Engine/Source/Editor/SceneOutliner/Private/SceneOutlinerSCCHandler.cpp b/Engine/Source/Editor/SceneOutliner/Private/SceneOutlinerSCCHandler.cpp
--- a/Engine/Source/Editor/SceneOutliner/Private/SceneOutlinerSCCHandler.cpp
+++ b/Engine/Source/Editor/SceneOutliner/Private/SceneOutlinerSCCHandler.cpp
@@ -6,6 +6,8 @@
 #include "Misc/MessageDialog.h"
 #include "FileHelpers.h"
 #include "SourceControlWindows.h"
+#include "ISourceControlWindowsModule.h"
+#include "UncontrolledChangelistsModule.h"
 #include "AssetViewUtils.h"
 
 #define LOCTEXT_NAMESPACE "FSceneOutlinerSCCHandler"
@@ -28,7 +30,7 @@ bool FSceneOutlinerSCCHandler::AddSourceControlMenuOptions(UToolMenu* Menu, TArr
 
 	CacheCanExecuteVars();
 
-	if (ISourceControlModule::Get().IsEnabled())
+	if (ISourceControlModule::Get().IsEnabled() || FUncontrolledChangelistsModule::Get().IsEnabled())
 	{
 		FToolMenuSection& Section = Menu->AddSection("AssetContextSourceControl");
 
@@ -83,8 +85,7 @@ void FSceneOutlinerSCCHandler::CacheCanExecuteVars()
 	bCanExecuteSCCCheckIn = false;
 	bCanExecuteSCCHistory = false;
 
-	ISourceControlProvider& SourceControlProvider = ISourceControlModule::Get().GetProvider();
-	if ( ISourceControlModule::Get().IsEnabled() )
+	if ( ISourceControlModule::Get().IsEnabled() || FUncontrolledChangelistsModule::Get().IsEnabled())
 	{
 		for (FSceneOutlinerTreeItemPtr SelectedItem : SelectedItems)
 		{
@@ -210,6 +211,17 @@ void FSceneOutlinerSCCHandler::FillSourceControlSubMenu(UToolMenu* Menu)
 			)
 		);
 	}
+
+	Section.AddSeparator(NAME_None);
+	Section.AddMenuEntry(
+		"SCCFindInChangelist",
+		LOCTEXT("SCCShowInChangelist", "Show in Changelist"),
+		LOCTEXT("SCCShowInChangelistTooltip", "Show the selected assets in the Changelist window."),
+		FSlateIcon(FAppStyle::GetAppStyleSetName(), "SourceControl.ChangelistsTab"),
+		FUIAction(
+			FExecuteAction::CreateSP(this, &FSceneOutlinerSCCHandler::ExecuteSCCShowInChangelist)
+		)
+	);
 }
 
 void FSceneOutlinerSCCHandler::GetSelectedPackageNames(TArray<FString>& OutPackageNames) const
@@ -246,6 +258,14 @@ void FSceneOutlinerSCCHandler::GetSelectedPackages(TArray<UPackage*>& OutPackage
 	}
 }
 
+void FSceneOutlinerSCCHandler::ExecuteSCCShowInChangelist()
+{
+	TArray<FString> PackageNames;
+	GetSelectedPackageNames(PackageNames);
+
+	ISourceControlWindowsModule::Get().SelectFiles(SourceControlHelpers::PackageFilenames(PackageNames));
+}
+
 void FSceneOutlinerSCCHandler::ExecuteSCCRefresh()
 {
 	TArray<FString> PackageNames;
diff --git a/Engine/Source/Editor/SceneOutliner/Public/SceneOutlinerSCCHandler.h b/Engine/Source/Editor/SceneOutliner/Public/SceneOutlinerSCCHandler.h
--- a/Engine/Source/Editor/SceneOutliner/Public/SceneOutlinerSCCHandler.h
+++ b/Engine/Source/Editor/SceneOutliner/Public/SceneOutlinerSCCHandler.h
@@ -31,6 +31,7 @@ private:
 	void ExecuteSCCCheckOut();
 	void ExecuteSCCCheckIn();
 	void ExecuteSCCHistory();
+	void ExecuteSCCShowInChangelist();
 
 	TArray<FSceneOutlinerTreeItemPtr> SelectedItems;
 
diff --git a/Engine/Source/Editor/SceneOutliner/Public/SceneOutlinerTreeItemSCC.h b/Engine/Source/Editor/SceneOutliner/Public/SceneOutlinerTreeItemSCC.h
--- a/Engine/Source/Editor/SceneOutliner/Public/SceneOutlinerTreeItemSCC.h
+++ b/Engine/Source/Editor/SceneOutliner/Public/SceneOutlinerTreeItemSCC.h
@@ -20,7 +20,7 @@ public:
 
 	FSourceControlStatePtr RefreshSourceControlState();
 
-	bool IsExternalPackage() { return ExternalPackage != nullptr; }
+	bool IsExternalPackage() { return !ExternalPackageName.IsEmpty(); }
 
 	FString GetPackageName() { return ExternalPackageName; }
 
diff --git a/Engine/Source/Editor/SceneOutliner/SceneOutliner.Build.cs b/Engine/Source/Editor/SceneOutliner/SceneOutliner.Build.cs
--- a/Engine/Source/Editor/SceneOutliner/SceneOutliner.Build.cs
+++ b/Engine/Source/Editor/SceneOutliner/SceneOutliner.Build.cs
@@ -28,6 +28,7 @@ public class SceneOutliner : ModuleRules
 				"SourceControl",
 				"EditorConfig",
 				"SourceControlWindows",
+				"UncontrolledChangelists",
 				"EditorWidgets",
 				"ToolWidgets",
 			}
diff --git a/Engine/Source/Editor/SourceControlWindows/Private/SSourceControlChangelists.cpp b/Engine/Source/Editor/SourceControlWindows/Private/SSourceControlChangelists.cpp
--- a/Engine/Source/Editor/SourceControlWindows/Private/SSourceControlChangelists.cpp
+++ b/Engine/Source/Editor/SourceControlWindows/Private/SSourceControlChangelists.cpp
@@ -738,6 +738,12 @@ void SSourceControlChangelistsWidget::OnRefresh()
 	ChangelistTreeView->RequestTreeRefresh();
 	UncontrolledChangelistTreeView->RequestTreeRefresh();
 	FileTreeView->RequestTreeRefresh();
+
+	if (FilesToSelect.Num() > 0)
+	{
+		TArray<FString> LocalFilesToSelect(MoveTemp(FilesToSelect));
+		SetSelectedFiles(LocalFilesToSelect);
+	}
 }
 
 void SSourceControlChangelistsWidget::OnSourceControlProviderChanged(ISourceControlProvider& OldProvider, ISourceControlProvider& NewProvider)
@@ -853,6 +859,130 @@ FSourceControlChangelistPtr SSourceControlChangelistsWidget::GetChangelistFromSe
 	return ChangelistState ? (FSourceControlChangelistPtr)(ChangelistState->GetChangelist()) : nullptr;
 }
 
+void SSourceControlChangelistsWidget::SetSelectedFiles(const TArray<FString>& Filenames)
+{
+	if (bShouldRefresh || bIsRefreshing)
+	{
+		FilesToSelect = Filenames;
+		return;
+	}
+
+	check(Filenames.Num() > 0);
+		
+	// Finds the Changelist tree item containing this Filename if it exists.
+	auto FindChangelist = [this](const FString& Filename) -> TSharedPtr<IChangelistTreeItem>
+	{
+		for (const TSharedPtr<IChangelistTreeItem>& Item : ChangelistTreeNodes)
+		{
+			for (const TSharedPtr<IChangelistTreeItem>& ChildItem : Item->GetChildren())
+			{
+				if (ChildItem->GetTreeItemType() == IChangelistTreeItem::File)
+				{
+					const FString& ChildFilename = StaticCastSharedPtr<FFileTreeItem>(ChildItem)->FileState->GetFilename();
+					if (ChildFilename.Compare(Filename, ESearchCase::IgnoreCase) == 0)
+					{
+						return Item;
+					}
+				}
+			}
+		}
+
+		for (const TSharedPtr<IChangelistTreeItem>& Item : UncontrolledChangelistTreeNodes)
+		{
+			for (const TSharedPtr<IChangelistTreeItem>& ChildItem : Item->GetChildren())
+			{
+				if (ChildItem->GetTreeItemType() == IChangelistTreeItem::File)
+				{
+					const FString& ChildFilename = StaticCastSharedPtr<FFileTreeItem>(ChildItem)->FileState->GetFilename();
+					if (ChildFilename.Compare(Filename, ESearchCase::IgnoreCase) == 0)
+					{
+						return Item;
+					}
+				}
+				else if (ChildItem->GetTreeItemType() == IChangelistTreeItem::OfflineFile)
+				{
+					const FString& ChildFilename = StaticCastSharedPtr<FOfflineFileTreeItem>(ChildItem)->GetFilename();
+					if (ChildFilename.Compare(Filename, ESearchCase::IgnoreCase) == 0)
+					{
+						return Item;
+					}
+				}
+			}
+		}
+
+		return nullptr;
+	};
+			
+	TSharedPtr<IChangelistTreeItem> FoundChangelistTreeItem = nullptr;
+	// Find filename in Changelist, since filenames might not be in same Changelist, start from the last Filename as it might be the last selected one and give it priority
+	for (int32 Index = Filenames.Num() - 1; Index >= 0; --Index)
+	{
+		if (TSharedPtr<IChangelistTreeItem> ChangelistTreeItem = FindChangelist(Filenames[Index]))
+		{
+			FoundChangelistTreeItem = ChangelistTreeItem;
+			break;
+		}
+	}
+
+	// If we found a Changelist, select files 
+	if (FoundChangelistTreeItem)
+	{
+		// Save current Selection/Expansion state
+		FExpandedAndSelectionStates State;
+		SaveExpandedAndSelectionStates(State);
+		State.SelectedChangelistNode = nullptr;
+		State.SelectedUncontrolledChangelistNode = nullptr;
+		State.SelectedFileNodes.Empty();
+
+		// To make search faster store all filenames lower case
+		TSet<FString> FilenamesLowerCase;
+		Algo::Transform(Filenames, FilenamesLowerCase, [](const FString& Filename) { return Filename.ToLower(); });
+
+		if (FoundChangelistTreeItem->GetTreeItemType() == IChangelistTreeItem::Changelist)
+		{
+			State.SelectedChangelistNode = StaticCastSharedPtr<FChangelistTreeItem>(FoundChangelistTreeItem);
+
+			for (const TSharedPtr<IChangelistTreeItem>& ChildItem : FoundChangelistTreeItem->GetChildren())
+			{
+				if (ChildItem->GetTreeItemType() == IChangelistTreeItem::File)
+				{
+					const FString& ChildFilename = StaticCastSharedPtr<FFileTreeItem>(ChildItem)->FileState->GetFilename().ToLower();
+					if (FilenamesLowerCase.Contains(ChildFilename))
+					{
+						State.SelectedFileNodes.Add(ChildItem);
+					}
+				}
+			}
+		}
+		else if (FoundChangelistTreeItem->GetTreeItemType() == IChangelistTreeItem::UncontrolledChangelist)
+		{
+			State.SelectedUncontrolledChangelistNode = StaticCastSharedPtr<FUncontrolledChangelistTreeItem>(FoundChangelistTreeItem);
+
+			for (const TSharedPtr<IChangelistTreeItem>& ChildItem : FoundChangelistTreeItem->GetChildren())
+			{
+				if (ChildItem->GetTreeItemType() == IChangelistTreeItem::File)
+				{
+					const FString& ChildFilename = StaticCastSharedPtr<FFileTreeItem>(ChildItem)->FileState->GetFilename().ToLower();
+					if (FilenamesLowerCase.Contains(ChildFilename))
+					{
+						State.SelectedFileNodes.Add(ChildItem);
+					}
+				}
+				else if (ChildItem->GetTreeItemType() == IChangelistTreeItem::OfflineFile)
+				{
+					const FString& ChildFilename = StaticCastSharedPtr<FOfflineFileTreeItem>(ChildItem)->GetFilename().ToLower();
+					if (FilenamesLowerCase.Contains(ChildFilename))
+					{
+						State.SelectedFileNodes.Add(ChildItem);
+					}
+				}
+			}
+		}
+
+		RestoreExpandedAndSelectionStates(State);
+	}
+}
+
 TArray<FString> SSourceControlChangelistsWidget::GetSelectedFiles()
 {
 	TArray<FChangelistTreeItemPtr> SelectedItems = FileTreeView->GetSelectedItems();
diff --git a/Engine/Source/Editor/SourceControlWindows/Private/SSourceControlChangelists.h b/Engine/Source/Editor/SourceControlWindows/Private/SSourceControlChangelists.h
--- a/Engine/Source/Editor/SourceControlWindows/Private/SSourceControlChangelists.h
+++ b/Engine/Source/Editor/SourceControlWindows/Private/SSourceControlChangelists.h
@@ -32,6 +32,9 @@ public:
 	/** Constructs the widget */
 	void Construct(const FArguments& InArgs);
 
+	/** Set selected files */
+	void SetSelectedFiles(const TArray<FString>& Filenames);
+
 private:
 	// Holds the list/state of selected and expanded items in the changelist views and file views.
 	struct FExpandedAndSelectionStates
@@ -169,6 +172,9 @@ private:
 	bool bShouldRefresh = false;
 	bool bSourceControlAvailable = false;
 
+	/** Files to select after refresh */
+	TArray<FString> FilesToSelect;
+
 	void StartRefreshStatus();
 	void TickRefreshStatus(double InDeltaTime);
 	void EndRefreshStatus();
diff --git a/Engine/Source/Editor/SourceControlWindows/Private/SourceControlWindowsModule.cpp b/Engine/Source/Editor/SourceControlWindows/Private/SourceControlWindowsModule.cpp
--- a/Engine/Source/Editor/SourceControlWindows/Private/SourceControlWindowsModule.cpp
+++ b/Engine/Source/Editor/SourceControlWindows/Private/SourceControlWindowsModule.cpp
@@ -37,7 +37,8 @@ public:
 
 	virtual void ShowChangelistsTab() override;
 	virtual bool CanShowChangelistsTab() const override;
-	
+	virtual void SelectFiles(const TArray<FString>& Filenames);
+
 	DECLARE_DERIVED_EVENT(FSourceControlWindowsModule, ISourceControlWindowsModule::FChangelistFileDoubleClickedEvent, FChangelistFileDoubleClickedEvent);
 	virtual FChangelistFileDoubleClickedEvent& OnChangelistFileDoubleClicked() override { return ChangelistFileDoubleClickedEvent; }
 
@@ -125,4 +126,16 @@ bool FSourceControlWindowsModule::CanShowChangelistsTab() const
 	return (SourceControlModule.IsEnabled() && SourceControlModule.GetProvider().IsAvailable() && SourceControlModule.GetProvider().UsesChangelists()) || FUncontrolledChangelistsModule::Get().IsEnabled();
 }
 
+void FSourceControlWindowsModule::SelectFiles(const TArray<FString>& Filenames)
+{
+	if (Filenames.Num() > 0 && CanShowChangelistsTab())
+	{
+		ShowChangelistsTab();
+		if (TSharedPtr<SSourceControlChangelistsWidget> ChangelistsWidgetPtr = ChangelistsWidget.Pin())
+		{
+			ChangelistsWidgetPtr->SetSelectedFiles(Filenames);
+		}
+	}
+}
+
 #undef LOCTEXT_NAMESPACE
diff --git a/Engine/Source/Editor/SourceControlWindows/Public/ISourceControlWindowsModule.h b/Engine/Source/Editor/SourceControlWindows/Public/ISourceControlWindowsModule.h
--- a/Engine/Source/Editor/SourceControlWindows/Public/ISourceControlWindowsModule.h
+++ b/Engine/Source/Editor/SourceControlWindows/Public/ISourceControlWindowsModule.h
@@ -26,6 +26,7 @@ public:
 
 	virtual void ShowChangelistsTab() = 0;
 	virtual bool CanShowChangelistsTab() const = 0;
+	virtual void SelectFiles(const TArray<FString>& Filenames) = 0;
 
 	DECLARE_EVENT_OneParam(ISourceControlWindowsModule, FChangelistFileDoubleClickedEvent, const FString&);
 	virtual FChangelistFileDoubleClickedEvent& OnChangelistFileDoubleClicked() = 0;
