Updated the changelist window to handle the Enter and Delete keys to Submit or Delete the currently selected changelist if conditions to do these operations are met. This mimics P4V behavior.
Updated the changelist window to handle double clicking on a changelist. This opens the submit dialog if the submit conditions are met. This mimics P4V behavior.

PR #9433: [Source Control] Changelists window: intercept Enter and Delete keys to Submit or Delete the current changelist (Contributed by [REDACTED_PERSON])

diff --git a/Engine/Plugins/Developer/PerforceSourceControl/Source/PerforceSourceControl/Private/PerforceSourceControlChangelist.h b/Engine/Plugins/Developer/PerforceSourceControl/Source/PerforceSourceControl/Private/PerforceSourceControlChangelist.h
--- a/Engine/Plugins/Developer/PerforceSourceControl/Source/PerforceSourceControl/Private/PerforceSourceControlChangelist.h
+++ b/Engine/Plugins/Developer/PerforceSourceControl/Source/PerforceSourceControl/Private/PerforceSourceControlChangelist.h
@@ -22,6 +22,11 @@ public:
 	{
 	}
 
+	virtual bool CanDelete() const override
+	{
+		return !IsDefault();
+	}
+
 	bool operator==(const FPerforceSourceControlChangelist& InOther) const
 	{
 		return ChangelistNumber == InOther.ChangelistNumber;
diff --git a/Engine/Source/Developer/SourceControl/Public/ISourceControlChangelist.h b/Engine/Source/Developer/SourceControl/Public/ISourceControlChangelist.h
--- a/Engine/Source/Developer/SourceControl/Public/ISourceControlChangelist.h
+++ b/Engine/Source/Developer/SourceControl/Public/ISourceControlChangelist.h
@@ -18,5 +18,11 @@ public:
 	/**
 	 * Virtual destructor
 	 */
-	virtual ~ISourceControlChangelist() {}
+	virtual ~ISourceControlChangelist() = default;
+
+	/**
+	 * Returns true if the changelist is deletable. Some special changelist might not be deletable like the
+	 * default Perforce changelist.
+	 */
+	virtual bool CanDelete() const { return true; }
 };
diff --git a/Engine/Source/Editor/SourceControlWindows/Private/SSourceControlChangelists.cpp b/Engine/Source/Editor/SourceControlWindows/Private/SSourceControlChangelists.cpp
--- a/Engine/Source/Editor/SourceControlWindows/Private/SSourceControlChangelists.cpp
+++ b/Engine/Source/Editor/SourceControlWindows/Private/SSourceControlChangelists.cpp
@@ -1158,6 +1158,43 @@ void SSourceControlChangelistsWidget::OnEndSourceControlOperation(const TSharedR
 	RefreshStatus = FText::GetEmpty(); // TODO: Should have a queue to stack async operations going on to correctly display concurrent async operations.
 }
 
+FReply SSourceControlChangelistsWidget::OnKeyDown(const FGeometry& MyGeometry, const FKeyEvent& InKeyEvent)
+{
+	FText FailureMessage;
+	if (InKeyEvent.GetKey() == EKeys::Enter)
+	{
+		// Submit the currently selected changelist (if any, and if conditions are met)
+		if (CanSubmitChangelist(&FailureMessage))
+		{
+			OnSubmitChangelist();
+		}
+		else
+		{
+			FText Title(LOCTEXT("Cannot_Submit_Changelist_From_Key_Title", "Cannot Submit Changelist"));
+			FMessageDialog::Open(EAppMsgType::Ok, EAppReturnType::Ok, FailureMessage, &Title);
+		}
+
+		return FReply::Handled();
+	}
+	else if (InKeyEvent.GetKey() == EKeys::Delete)
+	{
+		// Delete the currently selected changelist (if any, and if conditions are met)
+		if (CanDeleteChangelist(&FailureMessage))
+		{
+			OnDeleteChangelist();
+		}
+		else
+		{
+			FText Title(LOCTEXT("Cannot_Delete_Changelist_From_Key_Title", "Cannot Delete Changelist"));
+			FMessageDialog::Open(EAppMsgType::Ok, EAppReturnType::Ok, FailureMessage, &Title);
+		}
+
+		return FReply::Handled();
+	}
+
+	return FReply::Unhandled();
+}
+
 void SSourceControlChangelistsWidget::OnNewChangelist()
 {
 	FText ChangelistDescription;
@@ -1213,8 +1250,38 @@ void SSourceControlChangelistsWidget::OnDeleteChangelist()
 
 bool SSourceControlChangelistsWidget::CanDeleteChangelist()
 {
-	FSourceControlChangelistStatePtr Changelist = GetCurrentChangelistState();
-	return Changelist != nullptr && Changelist->GetFilesStates().Num() == 0 && Changelist->GetShelvedFilesStates().Num() == 0;
+	return CanDeleteChangelist(/*OutFailureMessage*/nullptr);
+}
+
+bool SSourceControlChangelistsWidget::CanDeleteChangelist(FText* OutFailureMessage)
+{
+	FSourceControlChangelistStatePtr ChangelistState = GetCurrentChangelistState();
+
+	if (ChangelistState == nullptr)
+	{
+		if (OutFailureMessage)
+		{
+			*OutFailureMessage = LOCTEXT("Cannot_Delete_No_Changelist", "No changelist selected.");
+		}
+		return false;
+	}
+	else if (!ChangelistState->GetChangelist()->CanDelete()) // Check if this changelist is deletable (ex. P4 default one is not deletable).
+	{
+		if (OutFailureMessage)
+		{
+			*OutFailureMessage = LOCTEXT("Cannot_Delete_Changelist_Not_Deletable", "The selected changelist cannot be deleted.");
+		}
+		return false;
+	}
+	else if (ChangelistState->GetFilesStates().Num() > 0 || ChangelistState->GetShelvedFilesStates().Num() > 0)
+	{
+		if (OutFailureMessage)
+		{
+			*OutFailureMessage = LOCTEXT("Cannot_Delete_Changelist_Not_Empty", "The changelist is not empty.");
+		}
+		return false;
+	}
+	return true;
 }
 
 void SSourceControlChangelistsWidget::OnEditChangelist()
@@ -1661,9 +1728,40 @@ void SSourceControlChangelistsWidget::OnSubmitChangelist()
 }
 
 bool SSourceControlChangelistsWidget::CanSubmitChangelist()
+{
+	return CanSubmitChangelist(/*OutFailureMessage*/nullptr);
+}
+
+bool SSourceControlChangelistsWidget::CanSubmitChangelist(FText* OutFailureMessage)
 {
 	FSourceControlChangelistStatePtr Changelist = GetCurrentChangelistState();
-	return Changelist != nullptr && Changelist->GetFilesStates().Num() > 0 && Changelist->GetShelvedFilesStates().Num() == 0;
+
+	if (Changelist == nullptr)
+	{
+		if (OutFailureMessage)
+		{
+			*OutFailureMessage = LOCTEXT("Cannot_Submit_Changelist_No_Selection", "No changelist selected.");
+		}
+		return false;
+	}
+	else if (Changelist->GetFilesStates().Num() <= 0)
+	{
+		if (OutFailureMessage)
+		{
+			*OutFailureMessage = LOCTEXT("Cannot_Submit_Changelist_No_Files", "The changelist doesn't contain any files to submit.");
+		}
+		return false;
+	}
+	else if (Changelist->GetShelvedFilesStates().Num() > 0)
+	{
+		if (OutFailureMessage)
+		{
+			*OutFailureMessage = LOCTEXT("Cannot_Submit_Changelist_Has_Shelved_Files", "The changelist contains shelved files.");
+		}
+		return false;
+	}
+
+	return true;
 }
 
 void SSourceControlChangelistsWidget::OnValidateChangelist()
@@ -2090,6 +2188,7 @@ TSharedRef<SChangelistTree> SSourceControlChangelistsWidget::CreateChangelistTre
 		.OnGenerateRow(this, &SSourceControlChangelistsWidget::OnGenerateRow)
 		.OnGetChildren(this, &SSourceControlChangelistsWidget::OnGetChangelistChildren)
 		.SelectionMode(ESelectionMode::Single)
+		.OnMouseButtonDoubleClick(this, &SSourceControlChangelistsWidget::OnItemDoubleClicked)
 		.OnContextMenuOpening(this, &SSourceControlChangelistsWidget::OnOpenContextMenu)
 		.OnSelectionChanged(this, &SSourceControlChangelistsWidget::OnChangelistSelectionChanged);
 }
@@ -2728,6 +2827,20 @@ void SSourceControlChangelistsWidget::OnItemDoubleClicked(TSharedPtr<IChangelist
 		const FString& Filename = StaticCastSharedPtr<FFileTreeItem>(Item)->FileState->GetFilename();
 		ISourceControlWindowsModule::Get().OnChangelistFileDoubleClicked().Broadcast(Filename);
 	}
+	else if (Item->GetTreeItemType() == IChangelistTreeItem::Changelist)
+	{
+		// Submit the currently selected changelists if conditions are met.
+		FText FailureMessage;
+		if (CanSubmitChangelist(&FailureMessage))
+		{
+			OnSubmitChangelist();
+		}
+		else
+		{
+			FText Title(LOCTEXT("Cannot_Submit_Changelist_Title", "Cannot Submit Changelist"));
+			FMessageDialog::Open(EAppMsgType::Ok, EAppReturnType::Ok, FailureMessage, &Title);
+		}
+	}
 }
 
 void SSourceControlChangelistsWidget::OnChangelistSelectionChanged(TSharedPtr<IChangelistTreeItem> SelectedChangelist, ESelectInfo::Type SelectionType)
diff --git a/Engine/Source/Editor/SourceControlWindows/Private/SSourceControlChangelists.h b/Engine/Source/Editor/SourceControlWindows/Private/SSourceControlChangelists.h
--- a/Engine/Source/Editor/SourceControlWindows/Private/SSourceControlChangelists.h
+++ b/Engine/Source/Editor/SourceControlWindows/Private/SSourceControlChangelists.h
@@ -89,13 +89,18 @@ private:
 	/** Returns list of currently selected shelved files */
 	TArray<FString> GetSelectedShelvedFiles();
 
+	/** Intercept Enter and Delete key presses to Submit or Delete the selected changelist (if conditions are met) */
+	virtual FReply OnKeyDown(const FGeometry& MyGeometry, const FKeyEvent& InKeyEvent) override;
+
 	/** Changelist operations */
 	void OnNewChangelist();
 	void OnDeleteChangelist();
 	bool CanDeleteChangelist();
+	bool CanDeleteChangelist(FText* OutFailureMessage);
 	void OnEditChangelist();
 	void OnSubmitChangelist();
 	bool CanSubmitChangelist();
+	bool CanSubmitChangelist(FText* OutFailureMessage);
 	void OnValidateChangelist();
 	bool CanValidateChangelist();
 
