Fixed the double popup issue that can happen in the changelist window for source control provider other than P4 that support changelist.

diff --git a/Engine/Source/Editor/SourceControlWindows/Private/SSourceControlChangelists.cpp b/Engine/Source/Editor/SourceControlWindows/Private/SSourceControlChangelists.cpp
--- a/Engine/Source/Editor/SourceControlWindows/Private/SSourceControlChangelists.cpp
+++ b/Engine/Source/Editor/SourceControlWindows/Private/SSourceControlChangelists.cpp
@@ -49,9 +49,20 @@ namespace
 /** Wraps the execution of a changelist operations with a slow task. */
 void ExecuteChangelistOperationWithSlowTaskWrapper(const FText& Message, const TFunction<void()>& ChangelistTask)
 {
-	FScopedSlowTask Progress(0.f, Message);
-	Progress.MakeDialog();
-	ChangelistTask();
+	// NOTE: This is a ugly workaround for P4 because the generic popup feedback operations in FScopedSourceControlProgress() was supressed for all synchrounous
+	//       operations. For other source control providers, the popup still shows up and showing a slow task and the FScopedSourceControlProgress at the same
+	//       time is a bad user experience. Until we fix source control popup situation in general in the Editor, this hack is in place to avoid the double popup.
+	//       At the time of writing, the other source control provider that supports changelists is Plastic.
+	if (ISourceControlModule::Get().GetProvider().GetName() == "Perforce")
+	{
+		FScopedSlowTask Progress(0.f, Message);
+		Progress.MakeDialog();
+		ChangelistTask();
+	}
+	else
+	{
+		ChangelistTask();
+	}
 }
 
 /** Wraps the execution of an uncontrolled changelist operations with a slow task. */
@@ -1136,10 +1147,11 @@ void SSourceControlChangelistsWidget::Execute(const FText& Message, const TShare
 	}
 	else
 	{
-		FScopedSlowTask Progress(0.f, Message);
-		Progress.MakeDialog();
-		ECommandResult::Type Result = SourceControlProvider.Execute(InOperation, InChangelist, InFiles, InConcurrency, InOperationCompleteDelegate);
-		OnEndSourceControlOperation(InOperation, Result);
+		ExecuteChangelistOperationWithSlowTaskWrapper(Message, [&]()
+		{
+			ECommandResult::Type Result = SourceControlProvider.Execute(InOperation, InChangelist, InFiles, InConcurrency, InOperationCompleteDelegate);
+			OnEndSourceControlOperation(InOperation, Result);
+		});
 	}
 }
 
@@ -2320,22 +2332,22 @@ protected:
 				TArray<FString> Files;
 				Algo::Transform(DropOperation->Files, Files, [](const FSourceControlStateRef& State) { return State->GetFilename(); });
 
-				FScopedSlowTask Progress(0.f, LOCTEXT("Dropping_Files_On_Changelist", "Moving file(s) to the selected changelist..."));
-				Progress.MakeDialog();
-
-				ISourceControlProvider& SourceControlProvider = ISourceControlModule::Get().GetProvider();
-				SourceControlProvider.Execute(ISourceControlOperation::Create<FMoveToChangelist>(), DestChangelist, Files, EConcurrency::Synchronous, FSourceControlOperationComplete::CreateLambda(
-					[](const TSharedRef<ISourceControlOperation>& Operation, ECommandResult::Type InResult)
-					{
-						if (InResult == ECommandResult::Succeeded)
-						{
-							DisplaySourceControlOperationNotification(LOCTEXT("Drop_Files_On_Changelist_Succeeded", "File(s) successfully moved to the selected changelist."), SNotificationItem::CS_Success);
-						}
-						else if (InResult == ECommandResult::Failed)
+				ExecuteChangelistOperationWithSlowTaskWrapper(LOCTEXT("Dropping_Files_On_Changelist", "Moving file(s) to the selected changelist..."), [&]()
+				{
+					ISourceControlProvider& SourceControlProvider = ISourceControlModule::Get().GetProvider();
+					SourceControlProvider.Execute(ISourceControlOperation::Create<FMoveToChangelist>(), DestChangelist, Files, EConcurrency::Synchronous, FSourceControlOperationComplete::CreateLambda(
+						[](const TSharedRef<ISourceControlOperation>& Operation, ECommandResult::Type InResult)
 						{
-							DisplaySourceControlOperationNotification(LOCTEXT("Drop_Files_On_Changelist_Failed", "Failed to move the file(s) to the selected changelist."), SNotificationItem::CS_Fail);
-						}
-					}));
+							if (InResult == ECommandResult::Succeeded)
+							{
+								DisplaySourceControlOperationNotification(LOCTEXT("Drop_Files_On_Changelist_Succeeded", "File(s) successfully moved to the selected changelist."), SNotificationItem::CS_Success);
+							}
+							else if (InResult == ECommandResult::Failed)
+							{
+								DisplaySourceControlOperationNotification(LOCTEXT("Drop_Files_On_Changelist_Failed", "Failed to move the file(s) to the selected changelist."), SNotificationItem::CS_Fail);
+							}
+						}));
+				});
 			}
 			else if (!DropOperation->UncontrolledFiles.IsEmpty())
 			{
