Adds Source Control File Dialog.
Displays source control file dialog when checkouting files is impossible during conversion from Uncontrolled Changelist to Changelist.

diff --git a/Engine/Source/Developer/UncontrolledChangelists/Private/UncontrolledChangelistsModule.cpp b/Engine/Source/Developer/UncontrolledChangelists/Private/UncontrolledChangelistsModule.cpp
--- a/Engine/Source/Developer/UncontrolledChangelists/Private/UncontrolledChangelistsModule.cpp
+++ b/Engine/Source/Developer/UncontrolledChangelists/Private/UncontrolledChangelistsModule.cpp
@@ -411,7 +411,7 @@ void FUncontrolledChangelistsModule::MoveFilesToUncontrolledChangelist(const TAr
 	}
 }
 
-void FUncontrolledChangelistsModule::MoveFilesToControlledChangelist(const TArray<FSourceControlStateRef>& InUncontrolledFileStates, const FSourceControlChangelistPtr& InChangelist)
+void FUncontrolledChangelistsModule::MoveFilesToControlledChangelist(const TArray<FSourceControlStateRef>& InUncontrolledFileStates, const FSourceControlChangelistPtr& InChangelist, TFunctionRef<bool(const TArray<FSourceControlStateRef>&)> InOpenConflictDialog)
 {
 	if (!IsEnabled())
 	{
@@ -421,10 +421,10 @@ void FUncontrolledChangelistsModule::MoveFilesToControlledChangelist(const TArra
 	TArray<FString> UncontrolledFilenames;
 	
 	Algo::Transform(InUncontrolledFileStates, UncontrolledFilenames, [](const FSourceControlStateRef& State) { return State->GetFilename(); });
-	MoveFilesToControlledChangelist(UncontrolledFilenames, InChangelist);
+	MoveFilesToControlledChangelist(UncontrolledFilenames, InChangelist, InOpenConflictDialog);
 }
 
-void FUncontrolledChangelistsModule::MoveFilesToControlledChangelist(const TArray<FString>& InUncontrolledFiles, const FSourceControlChangelistPtr& InChangelist)
+void FUncontrolledChangelistsModule::MoveFilesToControlledChangelist(const TArray<FString>& InUncontrolledFiles, const FSourceControlChangelistPtr& InChangelist, TFunctionRef<bool(const TArray<FSourceControlStateRef>&)> InOpenConflictDialog)
 {
 	if (!IsEnabled())
 	{
@@ -437,7 +437,7 @@ void FUncontrolledChangelistsModule::MoveFilesToControlledChangelist(const TArra
 	// Get updated filestates to check Checkout capabilities.
 	SourceControlProvider.GetState(InUncontrolledFiles, UpdatedFilestates, EStateCacheUsage::ForceUpdate);
 
-	TArray<UPackage*> PackageConflicts;
+	TArray<FSourceControlStateRef> FilesConflicts;
 	TArray<FString> FilesToAdd;
 	TArray<FString> FilesToCheckout;
 	TArray<FString> FilesToDelete;
@@ -445,26 +445,33 @@ void FUncontrolledChangelistsModule::MoveFilesToControlledChangelist(const TArra
 	// Check if we can Checkout files or mark for add
 	for (const FSourceControlStateRef& Filestate : UpdatedFilestates)
 	{
+		const FString& Filename = Filestate->GetFilename();
+
 		if (!Filestate->IsSourceControlled())
 		{
-			FilesToAdd.Add(Filestate->GetFilename());
+			FilesToAdd.Add(Filename);
 		}
-		else if (!IFileManager::Get().FileExists(*Filestate->GetFilename()))
+		else if (!IFileManager::Get().FileExists(*Filename))
 		{
-			FilesToDelete.Add(Filestate->GetFilename());
+			FilesToDelete.Add(Filename);
 		}
 		else if (Filestate->CanCheckout())
 		{
-			FilesToCheckout.Add(Filestate->GetFilename());
+			FilesToCheckout.Add(Filename);
+		}
+		else
+		{
+			FilesConflicts.Add(Filestate);
+			FilesToCheckout.Add(Filename);
 		}
 	}
 
 	bool bCanProceed = true;
 
 	// If we detected conflict, asking user if we should proceed.
-	if (!PackageConflicts.IsEmpty())
+	if (!FilesConflicts.IsEmpty())
 	{
-		bCanProceed = ShowConflictDialog(PackageConflicts);
+		bCanProceed = InOpenConflictDialog(FilesConflicts);
 	}
 
 	if (bCanProceed)
@@ -489,29 +496,6 @@ void FUncontrolledChangelistsModule::MoveFilesToControlledChangelist(const TArra
 	}
 }
 
-bool FUncontrolledChangelistsModule::ShowConflictDialog(TArray<UPackage*> InPackageConflicts)
-{
-	FPackagesDialogModule& CheckoutPackagesDialogModule = FModuleManager::LoadModuleChecked<FPackagesDialogModule>(TEXT("PackagesDialog"));
-
-	CheckoutPackagesDialogModule.CreatePackagesDialog(LOCTEXT("CheckoutPackagesDialogTitle", "Check Out Assets"),
-													  LOCTEXT("CheckoutPackagesDialogMessage", "Conflict detected in the following assets:"),
-													  true);
-
-	CheckoutPackagesDialogModule.AddButton(DRT_CheckOut, LOCTEXT("Dlg_CheckOutButton", "Check Out"), LOCTEXT("Dlg_CheckOutTooltip", "Attempt to Check Out Assets"));
-	CheckoutPackagesDialogModule.AddButton(DRT_Cancel, LOCTEXT("Dlg_Cancel", "Cancel"), LOCTEXT("Dlg_CancelTooltip", "Cancel Request"));
-
-	CheckoutPackagesDialogModule.SetWarning(LOCTEXT("CheckoutPackagesWarnMessage", "Warning: These assets are locked or not at the head revision. You may lose your changes if you continue, as you will be unable to submit them to source control."));
-
-	for (UPackage* Conflict : InPackageConflicts)
-	{
-		CheckoutPackagesDialogModule.AddPackageItem(Conflict, ECheckBoxState::Undetermined);
-	}
-
-	EDialogReturnType UserResponse = CheckoutPackagesDialogModule.ShowPackagesDialog();
-
-	return UserResponse == DRT_CheckOut;
-}
-
 void FUncontrolledChangelistsModule::OnStateChanged()
 {
 	OnUncontrolledChangelistModuleChanged.Broadcast();
diff --git a/Engine/Source/Developer/UncontrolledChangelists/Public/UncontrolledChangelistsModule.h b/Engine/Source/Developer/UncontrolledChangelists/Public/UncontrolledChangelistsModule.h
--- a/Engine/Source/Developer/UncontrolledChangelists/Public/UncontrolledChangelistsModule.h
+++ b/Engine/Source/Developer/UncontrolledChangelists/Public/UncontrolledChangelistsModule.h
@@ -114,15 +114,17 @@ public:
 	 * Moves files to a Controlled Changelist.
 	 * @param 	InUncontrolledFileStates 	The files to move.
 	 * @param 	InChangelist 				The Controlled Changelist where to move the files.
+	 * @param 	InOpenConflictDialog 		A callback to be used by the method when file conflicts are detected. The callback should display the files and ask the user if they should proceed.
 	 */
-	void MoveFilesToControlledChangelist(const TArray<FSourceControlStateRef>& InUncontrolledFileStates, const FSourceControlChangelistPtr& InChangelist);
+	void MoveFilesToControlledChangelist(const TArray<FSourceControlStateRef>& InUncontrolledFileStates, const FSourceControlChangelistPtr& InChangelist, TFunctionRef<bool(const TArray<FSourceControlStateRef>&)> InOpenConflictDialog);
 	
 	/**
 	 * Moves files to a Controlled Changelist.
 	 * @param 	InUncontrolledFiles 	The files to move.
 	 * @param 	InChangelist 			The Controlled Changelist where to move the files.
+	 * @param 	InOpenConflictDialog 	A callback to be used by the method when file conflicts are detected. The callback should display the files and ask the user if they should proceed.
 	 */
-	void MoveFilesToControlledChangelist(const TArray<FString>& InUncontrolledFiles, const FSourceControlChangelistPtr& InChangelist);
+	void MoveFilesToControlledChangelist(const TArray<FString>& InUncontrolledFiles, const FSourceControlChangelistPtr& InChangelist, TFunctionRef<bool(const TArray<FSourceControlStateRef>&)> InOpenConflictDialog);
 
 private:
 	/**
@@ -157,13 +159,6 @@ private:
 	 */
 	FString GetUObjectPackageFullpath(const UObject* InObject) const;
 
-	/**
-	 * Displays a Package Dialog warning the user about conflicting packages. 
-	 * @param 	InPackageConflicts 	The conflicting packages to display.
-	 * @return 	True if the user decided to proceed. False if they cancelled.
-	 */
-	bool ShowConflictDialog(TArray<UPackage*> InPackageConflicts);
-
 	/** Called when a state changed either in the module or an Uncontrolled Changelist. */
 	void OnStateChanged();
 
diff --git a/Engine/Source/Editor/SourceControlWindows/Private/SFileTableRow.h b/Engine/Source/Editor/SourceControlWindows/Private/SFileTableRow.h
new file mode 100644
--- /dev/null
+++ b/Engine/Source/Editor/SourceControlWindows/Private/SFileTableRow.h
@@ -0,0 +1,117 @@
+// Copyright Epic Games, Inc. All Rights Reserved.
+
+#pragma once
+
+#include "CoreMinimal.h"
+
+#include "Widgets/SCompoundWidget.h"
+#include "Widgets/Views/STreeView.h"
+
+#include "SSourceControlCommon.h"
+
+/** Display information about a file (icon, name, location, type, etc.) */
+class SFileTableRow : public SMultiColumnTableRow<FChangelistTreeItemPtr>
+{
+public:
+	SLATE_BEGIN_ARGS(SFileTableRow)
+		: _TreeItemToVisualize()
+	{}
+	SLATE_ARGUMENT(FChangelistTreeItemPtr, TreeItemToVisualize)
+		SLATE_EVENT(FOnDragDetected, OnDragDetected)
+		SLATE_END_ARGS()
+
+public:
+	/**
+	* Construct child widgets that comprise this widget.
+	*
+	* @param InArgs Declaration from which to construct this widget.
+	*/
+	void Construct(const FArguments& InArgs, const TSharedRef<STableViewBase>& InOwner)
+	{
+		TreeItem = static_cast<FFileTreeItem*>(InArgs._TreeItemToVisualize.Get());
+
+		FSuperRowType::FArguments Args = FSuperRowType::FArguments()
+			.OnDragDetected(InArgs._OnDragDetected)
+			.ShowSelection(true);
+		FSuperRowType::Construct(Args, InOwner);
+	}
+
+	// SMultiColumnTableRow overrides
+	virtual TSharedRef<SWidget> GenerateWidgetForColumn(const FName& ColumnName) override
+	{
+		if (ColumnName == TEXT("Icon"))
+		{
+			return SNew(SBox)
+				.WidthOverride(16) // Small Icons are usually 16x16
+				.HAlign(HAlign_Center)
+				[
+					SSourceControlCommon::GetSCCFileWidget(TreeItem->FileState, TreeItem->IsShelved())
+				];
+		}
+		else if (ColumnName == TEXT("Name"))
+		{
+			return SNew(STextBlock)
+				.Text(this, &SFileTableRow::GetDisplayName);
+		}
+		else if (ColumnName == TEXT("Path"))
+		{
+			return SNew(STextBlock)
+				.Text(this, &SFileTableRow::GetDisplayPath)
+				.ToolTipText(this, &SFileTableRow::GetFilename);
+		}
+		else if (ColumnName == TEXT("Type"))
+		{
+			return SNew(STextBlock)
+				.Text(this, &SFileTableRow::GetDisplayType)
+				.ColorAndOpacity(this, &SFileTableRow::GetDisplayColor);
+		}
+		else
+		{
+			return SNullWidget::NullWidget;
+		}
+	}
+
+	FText GetDisplayName() const
+	{
+		return TreeItem->GetAssetName();
+	}
+
+	FText GetFilename() const
+	{
+		return TreeItem->GetFileName();
+	}
+
+	FText GetDisplayPath() const
+	{
+		return TreeItem->GetAssetPath();
+	}
+
+	FText GetDisplayType() const
+	{
+		return TreeItem->GetAssetType();
+	}
+
+	FSlateColor GetDisplayColor() const
+	{
+		return TreeItem->GetAssetTypeColor();
+	}
+
+protected:
+	//~ Begin STableRow Interface.
+	virtual void OnDragEnter(FGeometry const& InGeometry, FDragDropEvent const& InDragDropEvent) override
+	{
+		TSharedPtr<FDragDropOperation> DragOperation = InDragDropEvent.GetOperation();
+		DragOperation->SetCursorOverride(EMouseCursor::SlashedCircle);
+	}
+
+	virtual void OnDragLeave(FDragDropEvent const& InDragDropEvent) override
+	{
+		TSharedPtr<FDragDropOperation> DragOperation = InDragDropEvent.GetOperation();
+		DragOperation->SetCursorOverride(EMouseCursor::None);
+	}
+	//~ End STableRow Interface.
+
+private:
+	/** The info about the widget that we are visualizing. */
+	FFileTreeItem* TreeItem;
+};
diff --git a/Engine/Source/Editor/SourceControlWindows/Private/SSourceControlChangelists.cpp b/Engine/Source/Editor/SourceControlWindows/Private/SSourceControlChangelists.cpp
--- a/Engine/Source/Editor/SourceControlWindows/Private/SSourceControlChangelists.cpp
+++ b/Engine/Source/Editor/SourceControlWindows/Private/SSourceControlChangelists.cpp
@@ -2,6 +2,7 @@
 
 #include "SSourceControlChangelists.h"
 
+#include "Editor.h"
 #include "Styling/AppStyle.h"
 #include "Algo/Transform.h"
 #include "Logging/MessageLog.h"
@@ -23,8 +24,10 @@
 #include "UncontrolledChangelistsModule.h"
 #include "SourceControlOperations.h"
 #include "ToolMenus.h"
+#include "SFileTableRow.h"
 #include "SSourceControlDescription.h"
 #include "SourceControlWindows.h"
+#include "SSourceControlFileDialog.h"
 #include "SourceControlHelpers.h"
 #include "SourceControlPreferences.h"
 #include "SourceControlMenuContext.h"
@@ -140,6 +143,36 @@ FText UpdateChangelistDescriptionToSubmitIfNeeded(const bool bInValidationResult
 	return InChangelistDescription;
 }
 
+bool OpenConflictDialog(const TArray<FSourceControlStateRef>& InFilesConflicts)
+{
+	TSharedPtr<SWindow> Window;
+	TSharedPtr<SSourceControlFileDialog> SourceControlFileDialog;
+
+	Window = SNew(SWindow)
+			 .Title(LOCTEXT("CheckoutPackagesDialogTitle", "Check Out Assets"))
+			 .SizingRule(ESizingRule::UserSized)
+			 .ClientSize(FVector2D(1024.0f, 512.0f))
+			 .SupportsMaximize(false)
+			 .SupportsMinimize(false)
+			 [
+			 	SNew(SBorder)
+			 	.Padding(4.f)
+			 	.BorderImage(FAppStyle::GetBrush("ToolPanel.GroupBorder"))
+			 	[
+			 		SAssignNew(SourceControlFileDialog, SSourceControlFileDialog)
+			 		.Message(LOCTEXT("CheckoutPackagesDialogMessage", "Conflict detected in the following assets:"))
+			 		.Warning(LOCTEXT("CheckoutPackagesWarnMessage", "Warning: These assets are locked or not at the head revision. You may lose your changes if you continue, as you will be unable to submit them to source control."))
+			 		.Files(InFilesConflicts)
+			 	]
+			 ];
+
+	SourceControlFileDialog->SetWindow(Window);
+	Window->SetWidgetToFocusOnActivate(SourceControlFileDialog);
+	GEditor->EditorAddModalWindow(Window.ToSharedRef());
+
+	return SourceControlFileDialog->IsProceedButtonPressed();
+}
+
 } // Anonymous namespace
 
 /** Implements drag and drop operation. */
@@ -1892,7 +1925,7 @@ void SSourceControlChangelistsWidget::OnMoveFiles()
 					// NOTE: Perform uncontrolled move only if the new changelist was created and the controlled file were move.
 					if ((!SelectedUncontrolledFiles.IsEmpty()) && static_cast<FNewChangelist&>(Operation.Get()).GetNewChangelist().IsValid())
 					{
-						FUncontrolledChangelistsModule::Get().MoveFilesToControlledChangelist(SelectedUncontrolledFiles, static_cast<FNewChangelist&>(Operation.Get()).GetNewChangelist());
+						FUncontrolledChangelistsModule::Get().MoveFilesToControlledChangelist(SelectedUncontrolledFiles, static_cast<FNewChangelist&>(Operation.Get()).GetNewChangelist(), OpenConflictDialog);
 					}
 
 					DisplaySourceControlOperationNotification(LOCTEXT("Move_Files_New_Changelist_Succeeded", "Files were successfully moved to a new changelist."), SNotificationItem::CS_Success);
@@ -1933,7 +1966,7 @@ void SSourceControlChangelistsWidget::OnMoveFiles()
 							// Perform an uncontrolled move only if the controlled file were move successfully.
 							if (!SelectedUncontrolledFiles.IsEmpty())
 							{
-								FUncontrolledChangelistsModule::Get().MoveFilesToControlledChangelist(SelectedUncontrolledFiles, Changelist);
+								FUncontrolledChangelistsModule::Get().MoveFilesToControlledChangelist(SelectedUncontrolledFiles, Changelist, OpenConflictDialog);
 							}
 
 							DisplaySourceControlOperationNotification(LOCTEXT("Move_Files_Between_Changelist_Succeeded", "File(s) successfully moved to the selected changelist."), SNotificationItem::CS_Success);
@@ -2355,7 +2388,7 @@ protected:
 				ExecuteUncontrolledChangelistOperationWithSlowTaskWrapper(LOCTEXT("Dropping_Uncontrolled_Files_On_Changelist", "Moving uncontrolled file(s) to the selected changelist..."), 
 					[&DropOperation, &DestChangelist]()
 				{
-					FUncontrolledChangelistsModule::Get().MoveFilesToControlledChangelist(DropOperation->UncontrolledFiles, DestChangelist);
+					FUncontrolledChangelistsModule::Get().MoveFilesToControlledChangelist(DropOperation->UncontrolledFiles, DestChangelist, OpenConflictDialog);
 					
 					// TODO: Fix MoveFilesToControlledChangelist() to report the possible errors and display a notification.
 				});
@@ -2528,113 +2561,6 @@ public:
 	}
 };
 
-/** Display information about a file (icon, name, location, type, etc.) */
-class SFileTableRow : public SMultiColumnTableRow<FChangelistTreeItemPtr>
-{
-public:
-	SLATE_BEGIN_ARGS(SFileTableRow)
-		: _TreeItemToVisualize()
-	{}
-	SLATE_ARGUMENT(FChangelistTreeItemPtr, TreeItemToVisualize)
-	SLATE_EVENT(FOnDragDetected, OnDragDetected)
-	SLATE_END_ARGS()
-
-public:
-	/**
-	* Construct child widgets that comprise this widget.
-	*
-	* @param InArgs Declaration from which to construct this widget.
-	*/
-	void Construct(const FArguments& InArgs, const TSharedRef<STableViewBase>& InOwner)
-	{
-		TreeItem = static_cast<FFileTreeItem*>(InArgs._TreeItemToVisualize.Get());
-
-		FSuperRowType::FArguments Args = FSuperRowType::FArguments()
-			.OnDragDetected(InArgs._OnDragDetected)
-			.ShowSelection(true);
-		FSuperRowType::Construct(Args, InOwner);
-	}
-
-	// SMultiColumnTableRow overrides
-	virtual TSharedRef<SWidget> GenerateWidgetForColumn(const FName& ColumnName) override
-	{
-		if (ColumnName == TEXT("Icon"))
-		{
-			return SNew(SBox)
-				.WidthOverride(16) // Small Icons are usually 16x16
-				.HAlign(HAlign_Center)
-				[
-					SSourceControlCommon::GetSCCFileWidget(TreeItem->FileState, TreeItem->IsShelved())
-				];
-		}
-		else if (ColumnName == TEXT("Name"))
-		{
-			return SNew(STextBlock)
-				.Text(this, &SFileTableRow::GetDisplayName);
-		}
-		else if (ColumnName == TEXT("Path"))
-		{
-			return SNew(STextBlock)
-				.Text(this, &SFileTableRow::GetDisplayPath)
-				.ToolTipText(this, &SFileTableRow::GetFilename);
-		}
-		else if (ColumnName == TEXT("Type"))
-		{
-			return SNew(STextBlock)
-				.Text(this, &SFileTableRow::GetDisplayType)
-				.ColorAndOpacity(this, &SFileTableRow::GetDisplayColor);
-		}
-		else
-		{
-			return SNullWidget::NullWidget;
-		}
-	}
-
-	FText GetDisplayName() const
-	{
-		return TreeItem->GetAssetName();
-	}
-
-	FText GetFilename() const
-	{
-		return TreeItem->GetFileName();
-	}
-
-	FText GetDisplayPath() const
-	{
-		return TreeItem->GetAssetPath();
-	}
-
-	FText GetDisplayType() const
-	{
-		return TreeItem->GetAssetType();
-	}
-
-	FSlateColor GetDisplayColor() const
-	{
-		return TreeItem->GetAssetTypeColor();
-	}
-
-protected:
-	//~ Begin STableRow Interface.
-	virtual void OnDragEnter(FGeometry const& InGeometry, FDragDropEvent const& InDragDropEvent) override
-	{
-		TSharedPtr<FDragDropOperation> DragOperation = InDragDropEvent.GetOperation();
-		DragOperation->SetCursorOverride(EMouseCursor::SlashedCircle);
-	}
-
-	virtual void OnDragLeave(FDragDropEvent const& InDragDropEvent) override
-	{
-		TSharedPtr<FDragDropOperation> DragOperation = InDragDropEvent.GetOperation();
-		DragOperation->SetCursorOverride(EMouseCursor::None);
-	}
-	//~ End STableRow Interface.
-
-private:
-	/** The info about the widget that we are visualizing. */
-	FFileTreeItem* TreeItem;
-};
-
 /** Display information about an offline file (icon, name, location, type, etc.). */
 class SOfflineFileTableRow : public SMultiColumnTableRow<FChangelistTreeItemPtr>
 {
diff --git a/Engine/Source/Editor/SourceControlWindows/Private/SSourceControlFileDialog.cpp b/Engine/Source/Editor/SourceControlWindows/Private/SSourceControlFileDialog.cpp
new file mode 100644
--- /dev/null
+++ b/Engine/Source/Editor/SourceControlWindows/Private/SSourceControlFileDialog.cpp
@@ -0,0 +1,208 @@
+// Copyright Epic Games, Inc. All Rights Reserved.
+
+#include "SSourceControlFileDialog.h"
+
+#include "Framework/Application/SlateApplication.h"
+#include "SFileTableRow.h"
+#include "Styling/AppStyle.h"
+#include "SWarningOrErrorBox.h"
+#include "Widgets/DeclarativeSyntaxSupport.h"
+#include "Widgets/Input/SButton.h"
+#include "Widgets/Layout/SBorder.h"
+#include "Widgets/Layout/SBox.h"
+#include "Widgets/Layout/SScrollBorder.h"
+
+#define LOCTEXT_NAMESPACE "SSourceControlFileDialog"
+
+void SSourceControlFileDialog::Construct(const FArguments& InArgs)
+{
+	Reset();
+	SetMessage(InArgs._Message);
+	SetWarning(InArgs._Warning);
+	SetFiles(InArgs._Files);
+
+	FileTreeView =	SNew(STreeView<FChangelistTreeItemPtr>)
+					.ItemHeight(24.0f)
+					.TreeItemsSource(&FileTreeNodes)
+					.OnGenerateRow(this, &SSourceControlFileDialog::OnGenerateRow)
+					.OnGetChildren(this, &SSourceControlFileDialog::OnGetFileChildren)
+					.HeaderRow
+					(
+						SNew(SHeaderRow)
+						+ SHeaderRow::Column("Icon")
+						.DefaultLabel(FText::GetEmpty())
+						.FillSized(18)
+
+						+ SHeaderRow::Column("Name")
+						.DefaultLabel(LOCTEXT("Name", "Name"))
+						.FillWidth(0.2f)
+
+						+ SHeaderRow::Column("Path")
+						.DefaultLabel(LOCTEXT("Path", "Path"))
+						.FillWidth(0.6f)
+
+						+ SHeaderRow::Column("Type")
+						.DefaultLabel(LOCTEXT("Type", "Type"))
+						.FillWidth(0.2f)
+					);
+
+	ChildSlot
+	[
+		SNew(SBorder)
+		.BorderImage(FAppStyle::Get().GetBrush("Brushes.Panel"))
+		.Padding(FMargin(16))
+		[
+			SNew(SVerticalBox)
+			+ SVerticalBox::Slot()
+			.FillHeight(0.1)
+			.Padding(0.0f, 0.0f, 0.0f, 8.0f)
+			[
+				SNew(STextBlock)
+				.Text(this, &SSourceControlFileDialog::GetMessage)
+				.AutoWrapText(true)
+			]
+			+ SVerticalBox::Slot()
+			.FillHeight(0.8)
+			[
+				SNew(SScrollBorder, FileTreeView.ToSharedRef())
+				[
+					FileTreeView.ToSharedRef()
+				]
+			]
+			+ SVerticalBox::Slot()
+			.Padding(0, 16.0f, 0, 0)
+			.AutoHeight()
+			[
+				SNew(SWarningOrErrorBox)
+				.Visibility(this, &SSourceControlFileDialog::GetWarningVisibility)
+				.Message(this, &SSourceControlFileDialog::GetWarning)
+			]
+			+ SVerticalBox::Slot()
+				.AutoHeight()
+				.Padding(0.0f, 16.0f, 0.0f, 0.0f)
+				.HAlign(HAlign_Right)
+				.VAlign(VAlign_Bottom)
+			[
+				SNew(SHorizontalBox)
+				+ SHorizontalBox::Slot()
+				.AutoWidth()
+				.Padding(5, 0)
+				[
+					SAssignNew(ProceedButton, SButton)
+					.ButtonStyle(&FAppStyle::Get(), "PrimaryButton")
+					.TextStyle(&FAppStyle::Get(), "PrimaryButtonText")
+					.Text(LOCTEXT("Proceed", "Proceed"))
+					.HAlign(HAlign_Right)
+					.VAlign(VAlign_Center)
+					.OnClicked(this, &SSourceControlFileDialog::OnProceedClicked)
+				]
+				+ SHorizontalBox::Slot()
+				.AutoWidth()
+				.Padding(5, 0)
+				[
+					SAssignNew(CancelButton, SButton)
+					.ButtonStyle(&FAppStyle::Get(), "Button")
+					.TextStyle(&FAppStyle::Get(), "ButtonText")
+					.Text(LOCTEXT("Cancel", "Cancel"))
+					.HAlign(HAlign_Right)
+					.VAlign(VAlign_Center)
+					.OnClicked(this, &SSourceControlFileDialog::OnCancelClicked)
+				]
+			]
+		]
+	];
+}
+
+void SSourceControlFileDialog::SetMessage(const FText& InMessage)
+{
+	Message = InMessage;
+}
+
+void SSourceControlFileDialog::SetWarning(const FText& InWarning)
+{
+	Warning = InWarning;
+}
+
+void SSourceControlFileDialog::SetFiles(const TArray<FSourceControlStateRef>& InFiles)
+{
+	for (const TSharedRef<ISourceControlState>& FileState : InFiles)
+	{
+		FileTreeNodes.Add(MakeShared<FFileTreeItem>(FileState, true));
+	}
+
+	if (FileTreeView.IsValid())
+	{
+		FileTreeView->RequestTreeRefresh();
+	}
+}
+
+TSharedRef<ITableRow> SSourceControlFileDialog::OnGenerateRow(TSharedPtr<IChangelistTreeItem> InItem, const TSharedRef<STableViewBase>& OwnerTable)
+{
+	return	SNew(SFileTableRow, OwnerTable)
+			.TreeItemToVisualize(InItem);
+}
+
+void SSourceControlFileDialog::Reset()
+{
+	bIsProceedButtonPressed = false;
+}
+
+void SSourceControlFileDialog::SetWindow(TSharedPtr<SWindow> InWindow)
+{
+	Window = MoveTemp(InWindow);
+}
+
+FReply SSourceControlFileDialog::OnKeyDown(const FGeometry& MyGeometry, const FKeyEvent& InKeyEvent)
+{
+	if ((InKeyEvent.GetKey() == EKeys::Enter && ProceedButton.IsValid()))
+	{
+		return OnProceedClicked();
+	}
+	else if ((InKeyEvent.GetKey() == EKeys::Escape && CancelButton.IsValid()))
+	{
+		return OnCancelClicked();
+	}
+
+	return FReply::Unhandled();
+}
+
+FText SSourceControlFileDialog::GetMessage() const
+{
+	return Message;
+}
+
+FText SSourceControlFileDialog::GetWarning() const
+{
+	return Warning;
+}
+
+EVisibility SSourceControlFileDialog::GetWarningVisibility() const
+{
+	return (Warning.IsEmpty()) ? EVisibility::Collapsed : EVisibility::HitTestInvisible;
+}
+
+FReply SSourceControlFileDialog::OnProceedClicked()
+{
+	bIsProceedButtonPressed = true;
+	CloseDialog();
+
+	return FReply::Handled();
+}
+
+FReply SSourceControlFileDialog::OnCancelClicked()
+{
+	bIsProceedButtonPressed = false;
+	CloseDialog();
+
+	return FReply::Handled();
+}
+
+void SSourceControlFileDialog::CloseDialog()
+{
+	if (Window.IsValid())
+	{
+		Window->RequestDestroyWindow();
+	}
+}
+
+#undef LOCTEXT_NAMESPACE
diff --git a/Engine/Source/Editor/SourceControlWindows/Private/SSourceControlFileDialog.h b/Engine/Source/Editor/SourceControlWindows/Private/SSourceControlFileDialog.h
new file mode 100644
--- /dev/null
+++ b/Engine/Source/Editor/SourceControlWindows/Private/SSourceControlFileDialog.h
@@ -0,0 +1,127 @@
+// Copyright Epic Games, Inc. All Rights Reserved.
+
+#pragma once
+
+#include "CoreMinimal.h"
+
+#include "SSourceControlCommon.h"
+#include "Widgets/SCompoundWidget.h"
+#include "Widgets/Views/STreeView.h"
+
+/**
+ * Represents dialog displaying a provided message along with source control files
+ */
+class SSourceControlFileDialog : public SCompoundWidget
+{
+public:
+	SLATE_BEGIN_ARGS(SSourceControlFileDialog)
+		: _Message()
+		, _Warning()
+		, _Files()
+	{}
+
+	/** The message of the widget */
+	SLATE_ARGUMENT(FText, Message)
+
+	/** The warning message of the widget */
+	SLATE_ARGUMENT(FText, Warning)
+
+	/** The Source Control Files to display */
+	SLATE_ARGUMENT(TArray<FSourceControlStateRef>, Files)
+
+	SLATE_END_ARGS()
+
+	/**
+	 * Construct this widget
+	 *
+	 * @param	InArgs	The declaration data for this widget
+	 */
+	void Construct(const FArguments& InArgs);
+
+	/**
+	 * Sets the message of the widget
+	 *
+	 * @param	InMessage	The string that the message should be set to
+	 */
+	void SetMessage(const FText& InMessage);
+
+	/**
+	 * Sets the warning message of the widget
+	 *
+	 * @param	InMessage	The string that the warning message should be set to
+	 */
+	void SetWarning(const FText& InMessage);
+
+	/**
+	 * Sets source control files to display
+	 *
+	 * @param	InFiles		The Source Control Files to display
+	 */
+	void SetFiles(const TArray<FSourceControlStateRef>& InFiles);
+
+	/** Reset the state of this dialogs buttons */
+	void Reset();
+
+	/** Returns true if the user pressed the primary button */
+	bool IsProceedButtonPressed() const { return bIsProceedButtonPressed; }
+
+	/**
+	 * Sets the owner's window handle
+	 * @param	InWindow	The window owning the dialog
+	 */
+	void SetWindow(TSharedPtr<SWindow> InWindow);
+
+private:
+	//~ Begin SWidget Interface
+	virtual FReply OnKeyDown(const FGeometry& MyGeometry, const FKeyEvent& InKeyEvent) override;
+	//~ End SWidget Interface
+
+	/** Treeview callback */
+	TSharedRef<ITableRow> OnGenerateRow(TSharedPtr<IChangelistTreeItem> InItem, const TSharedRef<STableViewBase>& OwnerTable);
+
+	/** Treeview callback */
+	void OnGetFileChildren(FChangelistTreeItemPtr InParent, TArray<FChangelistTreeItemPtr>& OutChildren) {}
+
+	/** Used to supply message text to the widget */
+	FText GetMessage() const;
+
+	/** Used to supply warning text to the widget */
+	FText GetWarning() const;
+
+	/** Used to determine visibility of the warning */
+	EVisibility GetWarningVisibility() const;
+
+	/** Called when the proceed button is clicked */
+	FReply OnProceedClicked();
+
+	/** Called when the cancel button is clicked */
+	FReply OnCancelClicked();
+
+	/** Closes the dialog by requesting the owning window to be destroyed */
+	void CloseDialog();
+
+private:
+	/** The Window owning the dialog */
+	TSharedPtr<SWindow> Window;
+
+	/** Hold the nodes displayed by the file tree. */
+	TArray<TSharedPtr<IChangelistTreeItem>> FileTreeNodes;
+
+	/** Display the list of files associated to the selected changelist, uncontrolled changelist or shelved node. */
+	TSharedPtr<STreeView<FChangelistTreeItemPtr>> FileTreeView;
+
+	/** The primary button widget. */
+	TSharedPtr<SButton> ProceedButton;
+
+	/** The secondary button widget. */
+	TSharedPtr<SButton> CancelButton;
+
+	/** The message to display */
+	FText Message;
+
+	/** The warning to display */
+	FText Warning;
+
+	/** Set to true if the user clicked on the primary button */
+	bool bIsProceedButtonPressed;
+};
diff --git a/Engine/Source/Editor/SourceControlWindows/SourceControlWindows.Build.cs b/Engine/Source/Editor/SourceControlWindows/SourceControlWindows.Build.cs
--- a/Engine/Source/Editor/SourceControlWindows/SourceControlWindows.Build.cs
+++ b/Engine/Source/Editor/SourceControlWindows/SourceControlWindows.Build.cs
@@ -21,6 +21,7 @@ public class SourceControlWindows : ModuleRules
 				"SourceControl", 
 				"UncontrolledChangelists",
 				"AssetTools",
+				"ToolWidgets",
 				"EditorFramework",
 				"WorkspaceMenuStructure",
 				"UnrealEd",		// We need this dependency here because we use PackageTools.
