Implemented 2-columns sort in the changelist window.
Added a tooltip over the 'icon' colum in the changelist window.

diff --git a/Engine/Source/Editor/SourceControlWindows/Private/SSourceControlChangelists.cpp b/Engine/Source/Editor/SourceControlWindows/Private/SSourceControlChangelists.cpp
--- a/Engine/Source/Editor/SourceControlWindows/Private/SSourceControlChangelists.cpp
+++ b/Engine/Source/Editor/SourceControlWindows/Private/SSourceControlChangelists.cpp
@@ -396,6 +396,8 @@ void SSourceControlChangelistsWidget::Construct(const FArguments& InArgs)
 	SourceControlStateChangedDelegateHandle = SCCModule.GetProvider().RegisterSourceControlStateChanged_Handle(FSourceControlStateChanged::FDelegate::CreateSP(this, &SSourceControlChangelistsWidget::OnSourceControlStateChanged));
 	UncontrolledChangelistModule.OnUncontrolledChangelistModuleChanged.AddSP(this, &SSourceControlChangelistsWidget::OnSourceControlStateChanged);
 
+	PrimarySortedColumn = "Name";
+
 	ChangelistTreeView = CreateChangelistTreeView(ChangelistTreeNodes);
 	UncontrolledChangelistTreeView = CreateChangelistTreeView(UncontrolledChangelistTreeNodes);
 	FileTreeView = CreateChangelistFilesView();
@@ -779,15 +781,18 @@ void SSourceControlChangelistsWidget::OnRefresh()
 	// Views were rebuilt from scratch, try expanding and selecting the nodes that were in that state before the update.
 	RestoreExpandedAndSelectionStates(ExpandedAndSelectedStates);
 
-	ChangelistTreeView->RequestTreeRefresh();
-	UncontrolledChangelistTreeView->RequestTreeRefresh();
-	FileTreeView->RequestTreeRefresh();
-
 	if (FilesToSelect.Num() > 0)
 	{
 		TArray<FString> LocalFilesToSelect(MoveTemp(FilesToSelect));
 		SetSelectedFiles(LocalFilesToSelect);
 	}
+
+	// Restore the sort order.
+	SortFileView();
+
+	ChangelistTreeView->RequestTreeRefresh();
+	UncontrolledChangelistTreeView->RequestTreeRefresh();
+	FileTreeView->RequestTreeRefresh();
 }
 
 void SSourceControlChangelistsWidget::OnSourceControlProviderChanged(ISourceControlProvider& OldProvider, ISourceControlProvider& NewProvider)
@@ -912,7 +917,7 @@ void SSourceControlChangelistsWidget::SetSelectedFiles(const TArray<FString>& Fi
 	}
 
 	check(Filenames.Num() > 0);
-		
+
 	// Finds the Changelist tree item containing this Filename if it exists.
 	auto FindChangelist = [this](const FString& Filename) -> TSharedPtr<IChangelistTreeItem>
 	{
@@ -956,7 +961,7 @@ void SSourceControlChangelistsWidget::SetSelectedFiles(const TArray<FString>& Fi
 
 		return nullptr;
 	};
-			
+
 	TSharedPtr<IChangelistTreeItem> FoundChangelistTreeItem = nullptr;
 	// Find filename in Changelist, since filenames might not be in same Changelist, start from the last Filename as it might be the last selected one and give it priority
 	for (int32 Index = Filenames.Num() - 1; Index >= 0; --Index)
@@ -1024,6 +1029,9 @@ void SSourceControlChangelistsWidget::SetSelectedFiles(const TArray<FString>& Fi
 		}
 
 		RestoreExpandedAndSelectionStates(State);
+
+		// Restore the sort order.
+		SortFileView();
 	}
 }
 
@@ -2254,18 +2262,30 @@ TSharedRef<STreeView<FChangelistTreeItemPtr>> SSourceControlChangelistsWidget::C
 			+SHeaderRow::Column("Icon")
 			.DefaultLabel(FText::GetEmpty())
 			.FillSized(18)
+			.SortPriority(this, &SSourceControlChangelistsWidget::GetColumnSortPriority, FName("Icon"))
+			.SortMode(this, &SSourceControlChangelistsWidget::GetColumnSortMode, FName("Icon"))
+			.OnSort(this, &SSourceControlChangelistsWidget::OnColumnSortModeChanged)
 
 			+SHeaderRow::Column("Name")
 			.DefaultLabel(LOCTEXT("Name", "Name"))
 			.FillWidth(0.2f)
+			.SortPriority(this, &SSourceControlChangelistsWidget::GetColumnSortPriority, FName("Name"))
+			.SortMode(this, &SSourceControlChangelistsWidget::GetColumnSortMode, FName("Name"))
+			.OnSort(this, &SSourceControlChangelistsWidget::OnColumnSortModeChanged)
 
 			+SHeaderRow::Column("Path")
 			.DefaultLabel(LOCTEXT("Path", "Path"))
 			.FillWidth(0.6f)
+			.SortPriority(this, &SSourceControlChangelistsWidget::GetColumnSortPriority, FName("Path"))
+			.SortMode(this, &SSourceControlChangelistsWidget::GetColumnSortMode, FName("Path"))
+			.OnSort(this, &SSourceControlChangelistsWidget::OnColumnSortModeChanged)
 
 			+SHeaderRow::Column("Type")
 			.DefaultLabel(LOCTEXT("Type", "Type"))
 			.FillWidth(0.2f)
+			.SortPriority(this, &SSourceControlChangelistsWidget::GetColumnSortPriority, FName("Type"))
+			.SortMode(this, &SSourceControlChangelistsWidget::GetColumnSortMode, FName("Type"))
+			.OnSort(this, &SSourceControlChangelistsWidget::OnColumnSortModeChanged)
 		);
 }
 
@@ -2822,6 +2842,183 @@ void SSourceControlChangelistsWidget::OnChangelistSelectionChanged(TSharedPtr<IC
 	FileTreeView->RequestTreeRefresh();
 }
 
+EColumnSortPriority::Type SSourceControlChangelistsWidget::GetColumnSortPriority(const FName ColumnId) const
+{
+	if (ColumnId == PrimarySortedColumn)
+	{
+		return EColumnSortPriority::Primary;
+	}
+	else if (ColumnId == SecondarySortedColumn)
+	{
+		return EColumnSortPriority::Secondary;
+	}
+
+	return EColumnSortPriority::Max; // No specific priority.
+}
+
+EColumnSortMode::Type SSourceControlChangelistsWidget::GetColumnSortMode(const FName ColumnId) const
+{
+	if (ColumnId == PrimarySortedColumn)
+	{
+		return PrimarySortMode;
+	}
+	else if (ColumnId == SecondarySortedColumn)
+	{
+		return SecondarySortMode;
+	}
+
+	return EColumnSortMode::None;
+}
+
+void SSourceControlChangelistsWidget::OnColumnSortModeChanged(const EColumnSortPriority::Type InSortPriority, const FName& InColumnId, const EColumnSortMode::Type InSortMode)
+{
+	if (InSortPriority == EColumnSortPriority::Primary)
+	{
+		PrimarySortedColumn = InColumnId;
+		PrimarySortMode = InSortMode;
+
+		if (InColumnId == SecondarySortedColumn) // Cannot be primary and secondary at the same time.
+		{
+			SecondarySortedColumn = FName();
+			SecondarySortMode = EColumnSortMode::None;
+		}
+	}
+	else if (InSortPriority == EColumnSortPriority::Secondary)
+	{
+		SecondarySortedColumn = InColumnId;
+		SecondarySortMode = InSortMode;
+	}
+
+	SortFileView();
+	FileTreeView->RequestListRefresh();
+}
+
+void SSourceControlChangelistsWidget::SortFileView()
+{
+	// Invoked when sorting the icons. This gives a priority to the status for sorting purpose.
+	auto GetSourceControlStateWeight = [](const ISourceControlState& State)
+	{
+		if (!State.IsCurrent())        { return 0; } // First if sorted in ascending order.
+		if (State.IsUnknown())         { return 1; }
+		if (State.IsConflicted())      { return 2; }
+		if (State.IsCheckedOutOther()) { return 3; }
+		if (State.IsCheckedOut())      { return 4; }
+		if (State.IsDeleted())         { return 5; }
+		if (State.IsAdded())           { return 6; }
+		else                           { return 7; }
+	};
+
+	// Implements the equivalent of operator< to compare the files status.
+	auto OperatorLessFileIcon = [&GetSourceControlStateWeight](const FFileTreeItem& Lhs, const FFileTreeItem& Rhs) -> bool
+	{
+		return GetSourceControlStateWeight(Lhs.FileState.Get()) < GetSourceControlStateWeight(Rhs.FileState.Get());
+	};
+
+	auto OperatorLessOfflineFileIcon = [](const FOfflineFileTreeItem& Lhs, const FOfflineFileTreeItem& Rhs) -> bool
+	{
+		return false; // Offline file do not have state, don't change the order.
+	};
+
+	auto OperatorLessIcon = [&OperatorLessFileIcon, &OperatorLessOfflineFileIcon](IChangelistTreeItem* Lhs, IChangelistTreeItem* Rhs)-> bool
+	{
+		check(Lhs->GetTreeItemType() == Rhs->GetTreeItemType());
+		switch (Lhs->GetTreeItemType())
+		{
+		case IChangelistTreeItem::File:
+		case IChangelistTreeItem::ShelvedFile:
+			return OperatorLessFileIcon(*static_cast<FFileTreeItem*>(Lhs), *static_cast<FFileTreeItem*>(Rhs));
+		case IChangelistTreeItem::OfflineFile:
+			return OperatorLessOfflineFileIcon(*static_cast<FOfflineFileTreeItem*>(Lhs), *static_cast<FOfflineFileTreeItem*>(Rhs));
+		default:
+			return false;
+		}
+	};
+
+	auto GetName = [](IChangelistTreeItem* Item) -> FString
+	{
+		switch (Item->GetTreeItemType())
+		{
+		case IChangelistTreeItem::File:
+		case IChangelistTreeItem::ShelvedFile:
+			return static_cast<FFileTreeItem*>(Item)->GetAssetName().ToString();
+		case IChangelistTreeItem::OfflineFile:
+			return static_cast<FOfflineFileTreeItem*>(Item)->GetDisplayName().ToString();
+		default:
+			return FString();
+		}
+	};
+
+	auto GetPath = [](IChangelistTreeItem* Item) -> FString
+	{
+		switch (Item->GetTreeItemType())
+		{
+		case IChangelistTreeItem::File:
+		case IChangelistTreeItem::ShelvedFile:
+			return static_cast<FFileTreeItem*>(Item)->GetAssetPath().ToString();
+		case IChangelistTreeItem::OfflineFile:
+			return static_cast<FOfflineFileTreeItem*>(Item)->GetDisplayPath().ToString();
+		default:
+			return FString();
+		}
+	};
+
+	auto GetType = [](IChangelistTreeItem* Item)  -> FString
+	{
+		switch (Item->GetTreeItemType())
+		{
+		case IChangelistTreeItem::File:
+		case IChangelistTreeItem::ShelvedFile:
+			return static_cast<FFileTreeItem*>(Item)->GetAssetType().ToString();
+		case IChangelistTreeItem::OfflineFile:
+			return static_cast<FOfflineFileTreeItem*>(Item)->GetDisplayType().ToString();
+		default:
+			return FString();
+		}
+	};
+
+	auto Compare = [&](const TSharedPtr<IChangelistTreeItem>& Lhs, const TSharedPtr<IChangelistTreeItem>& Rhs, const FName& ColName, EColumnSortMode::Type SortMode) -> bool
+	{
+		if (ColName == "Icon")
+		{
+			return SortMode == EColumnSortMode::Ascending ? OperatorLessIcon(Lhs.Get(), Rhs.Get()) : OperatorLessIcon(Rhs.Get(), Lhs.Get());
+		}
+		else if (ColName == "Name")
+		{
+			return SortMode == EColumnSortMode::Ascending ? GetName(Lhs.Get()) < GetName(Rhs.Get()) : GetName(Lhs.Get()) > GetName(Rhs.Get());
+		}
+		else if (ColName == "Path")
+		{
+			return SortMode == EColumnSortMode::Ascending ? GetPath(Lhs.Get()) < GetPath(Rhs.Get()) : GetPath(Lhs.Get()) > GetPath(Rhs.Get());
+		}
+		else if (ColName == "Type")
+		{
+			return SortMode == EColumnSortMode::Ascending ? GetType(Lhs.Get()) < GetType(Rhs.Get()) : GetType(Lhs.Get()) > GetType(Rhs.Get());
+		}
+		else
+		{
+			checkNoEntry();
+			return false;
+		}
+	};
+
+	FileTreeNodes.Sort([&](const TSharedPtr<IChangelistTreeItem>& Lhs, const TSharedPtr<IChangelistTreeItem>& Rhs)
+	{
+		if (Compare(Lhs, Rhs, PrimarySortedColumn, PrimarySortMode))
+		{
+			return true; // Lhs must be before Rhs based on the primary sort order.
+		}
+		else if (Compare(Rhs, Lhs, PrimarySortedColumn, PrimarySortMode)) // Invert operands order (goal is to check if operands are equal or not)
+		{
+			return false; // Rhs must be before Lhs based on the primary sort.
+		}
+		else if (!SecondarySortedColumn.IsNone()) // Lhs == Rhs on the primary column, need to sort according the secondary column if one is set.
+		{
+			return Compare(Lhs, Rhs, SecondarySortedColumn, SecondarySortMode);
+		}
+		return false;
+	});
+}
+
 void SSourceControlChangelistsWidget::SaveExpandedAndSelectionStates(FExpandedAndSelectionStates& OutStates)
 {
 	// Save the selected item from the 'changelists' tree.
diff --git a/Engine/Source/Editor/SourceControlWindows/Private/SSourceControlChangelists.h b/Engine/Source/Editor/SourceControlWindows/Private/SSourceControlChangelists.h
--- a/Engine/Source/Editor/SourceControlWindows/Private/SSourceControlChangelists.h
+++ b/Engine/Source/Editor/SourceControlWindows/Private/SSourceControlChangelists.h
@@ -54,6 +54,11 @@ private:
 	void OnGetFileChildren(FChangelistTreeItemPtr InParent, TArray<FChangelistTreeItemPtr>& OutChildren);
 	void OnGetChangelistChildren(FChangelistTreeItemPtr InParent, TArray<FChangelistTreeItemPtr>& OutChildren);
 
+	EColumnSortPriority::Type GetColumnSortPriority(const FName ColumnId) const;
+	EColumnSortMode::Type GetColumnSortMode(const FName ColumnId) const;
+	void OnColumnSortModeChanged(const EColumnSortPriority::Type SortPriority, const FName& ColumnId, const EColumnSortMode::Type InSortMode);
+	void SortFileView();
+
 	FReply OnFilesDragged(const FGeometry& InGeometry, const FPointerEvent& InMouseEvent);
 
 	void RequestRefresh();
@@ -153,7 +158,6 @@ private:
 	void EditChangelistDescription(const FText& InNewChangelistDescription, const FSourceControlChangelistStatePtr& InChangelistState);
 
 private:
-
 	TSharedPtr<SExpandableChangelistArea> ChangelistExpandableArea;
 	TSharedPtr<SExpandableChangelistArea> UncontrolledChangelistExpandableArea;
 
@@ -180,6 +184,11 @@ private:
 	/** Files to select after refresh */
 	TArray<FString> FilesToSelect;
 
+	FName PrimarySortedColumn;
+	FName SecondarySortedColumn;
+	EColumnSortMode::Type PrimarySortMode = EColumnSortMode::Ascending;
+	EColumnSortMode::Type SecondarySortMode = EColumnSortMode::None;
+
 	void StartRefreshStatus();
 	void TickRefreshStatus(double InDeltaTime);
 	void EndRefreshStatus();
diff --git a/Engine/Source/Editor/SourceControlWindows/Private/SSourceControlCommon.cpp b/Engine/Source/Editor/SourceControlWindows/Private/SSourceControlCommon.cpp
--- a/Engine/Source/Editor/SourceControlWindows/Private/SSourceControlCommon.cpp
+++ b/Engine/Source/Editor/SourceControlWindows/Private/SSourceControlCommon.cpp
@@ -320,6 +320,7 @@ TSharedRef<SWidget> GetSCCFileWidget(FSourceControlStateRef InFileState, bool bI
 			.HeightOverride(IconOverlaySize)
 			[
 				SNew(SLayeredImage, InFileState->GetIcon())
+				.ToolTipText(InFileState->GetDisplayTooltip())
 			]
 		];
 }
