Updated the source control changelist code to be more reusable.
 - Moved all source control table row widgets found in SSourceControlChangelists.cpp into a separated files SSourceControlChangelistRows.h/.cpp
 - Moved content SFileTableRow.h into SSourceControlChangelistRows.h/.cpp and deleted the file.
 - Created column Id constants for the changelist files "Icon", "Name", "Path", "Type" rather than using hardcoded string.
 - Moved functions shared by the SSourceCOntrolChangeslists.cpp and SSourceControlChangelistRows.cpp into SSourceControlCommon.h/.cpp.

diff --git a/Engine/Source/Editor/SourceControlWindows/Private/SFileTableRow.h b/Engine/Source/Editor/SourceControlWindows/Private/SFileTableRow.h
deleted file mode 100644
diff --git a/Engine/Source/Editor/SourceControlWindows/Private/SSourceControlChangelistRows.cpp b/Engine/Source/Editor/SourceControlWindows/Private/SSourceControlChangelistRows.cpp
new file mode 100644
--- /dev/null
+++ b/Engine/Source/Editor/SourceControlWindows/Private/SSourceControlChangelistRows.cpp
@@ -0,0 +1,391 @@
+// Copyright Epic Games, Inc. All Rights Reserved.
+
+#include "SSourceControlChangelistRows.h"
+
+#include "ISourceControlModule.h"
+#include "UncontrolledChangelistsModule.h"
+#include "SourceControlOperations.h"
+
+#define LOCTEXT_NAMESPACE "SourceControlChangelistRow"
+
+
+const FName SourceControlFileViewColumnId::Icon = TEXT("Icon");
+const FName SourceControlFileViewColumnId::Name = TEXT("Name");
+const FName SourceControlFileViewColumnId::Path = TEXT("Path");
+const FName SourceControlFileViewColumnId::Type = TEXT("Type");
+
+
+void SChangelistTableRow::Construct(const FArguments& InArgs, const TSharedRef<STableViewBase>& InOwner)
+{
+	TreeItem = static_cast<FChangelistTreeItem*>(InArgs._TreeItemToVisualize.Get());
+	OnPostDrop = InArgs._OnPostDrop;
+
+	const FSlateBrush* IconBrush = (TreeItem != nullptr) ?
+		FAppStyle::GetBrush(TreeItem->ChangelistState->GetSmallIconName()) :
+		FAppStyle::GetBrush("SourceControl.Changelist");
+
+	SetToolTipText(GetChangelistDescriptionText());
+
+	STableRow<FChangelistTreeItemPtr>::Construct(
+		STableRow<FChangelistTreeItemPtr>::FArguments()
+		.Style(FAppStyle::Get(), "TableView.Row")
+		.Content()
+		[
+			SNew(SHorizontalBox)
+			+SHorizontalBox::Slot() // Icon
+			.AutoWidth()
+			[
+				SNew(SImage)
+				.Image(IconBrush)
+			]
+			+SHorizontalBox::Slot() // Changelist number.
+			.Padding(2, 0, 0, 0)
+			.AutoWidth()
+			[
+				SNew(STextBlock)
+				.Text(this, &SChangelistTableRow::GetChangelistText)
+			]
+			+SHorizontalBox::Slot() // Files count.
+			.Padding(4, 0, 4, 0)
+			.AutoWidth()
+			[
+				SNew(STextBlock)
+				.Text(FText::Format(INVTEXT("({0})"), TreeItem->GetFileCount()))
+			]
+			+SHorizontalBox::Slot()
+			.Padding(2, 0, 0, 0)
+			.AutoWidth()
+			[
+				SNew(STextBlock)
+				.Text(this, &SChangelistTableRow::GetChangelistDescriptionText)
+			]
+		], InOwner);
+}
+
+FText SChangelistTableRow::GetChangelistText() const
+{
+	return TreeItem->GetDisplayText();
+}
+
+FText SChangelistTableRow::GetChangelistDescriptionText() const
+{
+	FString DescriptionString = TreeItem->GetDescriptionText().ToString();
+	// Here we'll both remove \r\n (when edited from the dialog) and \n (when we get it from the SCC)
+	DescriptionString.ReplaceInline(TEXT("\r"), TEXT(""));
+	DescriptionString.ReplaceInline(TEXT("\n"), TEXT(" "));
+	DescriptionString.TrimEndInline();
+	return FText::FromString(DescriptionString);
+}
+
+FReply SChangelistTableRow::OnDrop(const FGeometry& InGeometry, const FDragDropEvent& InDragDropEvent)
+{
+	TSharedPtr<FSCCFileDragDropOp> DropOperation = InDragDropEvent.GetOperationAs<FSCCFileDragDropOp>();
+	if (DropOperation.IsValid())
+	{
+		FSourceControlChangelistPtr DestChangelist = TreeItem->ChangelistState->GetChangelist();
+		check(DestChangelist.IsValid());
+
+		// NOTE: The UI don't show 'source controlled files' and 'uncontrolled files' at the same time. User cannot select and drag/drop both file types at the same time.
+		if (!DropOperation->Files.IsEmpty())
+		{
+			TArray<FString> Files;
+			Algo::Transform(DropOperation->Files, Files, [](const FSourceControlStateRef& State) { return State->GetFilename(); });
+
+			SSourceControlCommon::ExecuteChangelistOperationWithSlowTaskWrapper(LOCTEXT("Dropping_Files_On_Changelist", "Moving file(s) to the selected changelist..."), [&]()
+			{
+				ISourceControlProvider& SourceControlProvider = ISourceControlModule::Get().GetProvider();
+				SourceControlProvider.Execute(ISourceControlOperation::Create<FMoveToChangelist>(), DestChangelist, Files, EConcurrency::Synchronous, FSourceControlOperationComplete::CreateLambda(
+					[](const TSharedRef<ISourceControlOperation>& Operation, ECommandResult::Type InResult)
+					{
+						if (InResult == ECommandResult::Succeeded)
+						{
+							SSourceControlCommon::DisplaySourceControlOperationNotification(LOCTEXT("Drop_Files_On_Changelist_Succeeded", "File(s) successfully moved to the selected changelist."), SNotificationItem::CS_Success);
+						}
+						else if (InResult == ECommandResult::Failed)
+						{
+							SSourceControlCommon::DisplaySourceControlOperationNotification(LOCTEXT("Drop_Files_On_Changelist_Failed", "Failed to move the file(s) to the selected changelist."), SNotificationItem::CS_Fail);
+						}
+					}));
+			});
+		}
+		else if (!DropOperation->UncontrolledFiles.IsEmpty())
+		{
+			// NOTE: This function does several operations that can fails but we don't get feedback.
+			SSourceControlCommon::ExecuteUncontrolledChangelistOperationWithSlowTaskWrapper(LOCTEXT("Dropping_Uncontrolled_Files_On_Changelist", "Moving uncontrolled file(s) to the selected changelist..."), 
+				[&DropOperation, &DestChangelist]()
+			{
+				FUncontrolledChangelistsModule::Get().MoveFilesToControlledChangelist(DropOperation->UncontrolledFiles, DestChangelist, SSourceControlCommon::OpenConflictDialog);
+					
+				// TODO: Fix MoveFilesToControlledChangelist() to report the possible errors and display a notification.
+			});
+
+			OnPostDrop.ExecuteIfBound();
+		}
+	}
+
+	return FReply::Handled();
+}
+
+
+void SUncontrolledChangelistTableRow::Construct(const FArguments& InArgs, const TSharedRef<STableViewBase>& InOwner)
+{
+	TreeItem = static_cast<FUncontrolledChangelistTreeItem*>(InArgs._TreeItemToVisualize.Get());
+	OnPostDrop = InArgs._OnPostDrop;
+
+		const FSlateBrush* IconBrush = (TreeItem != nullptr) ?
+			FAppStyle::GetBrush(TreeItem->UncontrolledChangelistState->GetSmallIconName()) :
+			FAppStyle::GetBrush("SourceControl.Changelist");
+
+	SetToolTipText(GetChangelistDescriptionText());
+
+	STableRow<FChangelistTreeItemPtr>::Construct(
+		STableRow<FChangelistTreeItemPtr>::FArguments()
+		.Style(FAppStyle::Get(), "TableView.Row")
+		.Content()
+		[
+			SNew(SHorizontalBox)
+			+ SHorizontalBox::Slot()
+			.AutoWidth()
+			.VAlign(VAlign_Center)
+			[
+				SNew(SImage)
+				.Image(IconBrush)
+			]
+			+ SHorizontalBox::Slot()
+			.AutoWidth()
+			.Padding(2.0f, 0.0f)
+			.VAlign(VAlign_Center)
+			[
+				SNew(STextBlock)
+				.Text(this, &SUncontrolledChangelistTableRow::GetChangelistText)
+			]
+			+SHorizontalBox::Slot() // Files/Offline file count.
+			.Padding(4, 0, 4, 0)
+			.AutoWidth()
+			[
+				SNew(STextBlock)
+				.Text(FText::Format(INVTEXT("({0})"), TreeItem->GetFileCount() + TreeItem->GetOfflineFileCount()))
+			]
+			+ SHorizontalBox::Slot()
+			.AutoWidth()
+			.Padding(2.0f, 0.0f)
+			.VAlign(VAlign_Center)
+			[
+				SNew(STextBlock)
+				.Text(this, &SUncontrolledChangelistTableRow::GetChangelistDescriptionText)
+			]
+		], InOwner);
+}
+
+FText SUncontrolledChangelistTableRow::GetChangelistText() const
+{
+	return TreeItem->GetDisplayText();
+}
+
+FText SUncontrolledChangelistTableRow::GetChangelistDescriptionText() const
+{
+	FString DescriptionString = TreeItem->GetDescriptionText().ToString();
+	// Here we'll both remove \r\n (when edited from the dialog) and \n (when we get it from the SCC)
+	DescriptionString.ReplaceInline(TEXT("\r"), TEXT(""));
+	DescriptionString.ReplaceInline(TEXT("\n"), TEXT(" "));
+	DescriptionString.TrimEndInline();
+	return FText::FromString(DescriptionString);
+}
+
+FReply SUncontrolledChangelistTableRow::OnDrop(const FGeometry& InGeometry, const FDragDropEvent& InDragDropEvent)
+{
+	TSharedPtr<FSCCFileDragDropOp> Operation = InDragDropEvent.GetOperationAs<FSCCFileDragDropOp>();
+
+	if (Operation.IsValid())
+	{
+		SSourceControlCommon::ExecuteUncontrolledChangelistOperationWithSlowTaskWrapper(LOCTEXT("Drag_File_To_Uncontrolled_Changelist", "Moving file(s) to the selected uncontrolled changelists..."),
+			[this, &Operation]()
+		{
+			FUncontrolledChangelistsModule::Get().MoveFilesToUncontrolledChangelist(Operation->Files, Operation->UncontrolledFiles, TreeItem->UncontrolledChangelistState->Changelist);
+		});
+
+		OnPostDrop.ExecuteIfBound();
+	}
+
+	return FReply::Handled();
+}
+
+
+void SFileTableRow::Construct(const FArguments& InArgs, const TSharedRef<STableViewBase>& InOwner)
+{
+	TreeItem = static_cast<FFileTreeItem*>(InArgs._TreeItemToVisualize.Get());
+
+	FSuperRowType::FArguments Args = FSuperRowType::FArguments()
+		.OnDragDetected(InArgs._OnDragDetected)
+		.ShowSelection(true);
+	FSuperRowType::Construct(Args, InOwner);
+}
+
+TSharedRef<SWidget> SFileTableRow::GenerateWidgetForColumn(const FName& ColumnId)
+{
+	if (ColumnId == SourceControlFileViewColumnId::Icon)
+	{
+		return SNew(SBox)
+			.WidthOverride(16) // Small Icons are usually 16x16
+			.HAlign(HAlign_Center)
+			[
+				SSourceControlCommon::GetSCCFileWidget(TreeItem->FileState, TreeItem->IsShelved())
+			];
+	}
+	else if (ColumnId == SourceControlFileViewColumnId::Name)
+	{
+		return SNew(STextBlock)
+			.Text(this, &SFileTableRow::GetDisplayName);
+	}
+	else if (ColumnId == SourceControlFileViewColumnId::Path)
+	{
+		return SNew(STextBlock)
+			.Text(this, &SFileTableRow::GetDisplayPath)
+			.ToolTipText(this, &SFileTableRow::GetFilename);
+	}
+	else if (ColumnId == SourceControlFileViewColumnId::Type)
+	{
+		return SNew(STextBlock)
+			.Text(this, &SFileTableRow::GetDisplayType)
+			.ColorAndOpacity(this, &SFileTableRow::GetDisplayColor);
+	}
+	else
+	{
+		return SNullWidget::NullWidget;
+	}
+}
+
+FText SFileTableRow::GetDisplayName() const
+{
+	return TreeItem->GetAssetName();
+}
+
+FText SFileTableRow::GetFilename() const
+{
+	return TreeItem->GetFileName();
+}
+
+FText SFileTableRow::GetDisplayPath() const
+{
+	return TreeItem->GetAssetPath();
+}
+
+FText SFileTableRow::GetDisplayType() const
+{
+	return TreeItem->GetAssetType();
+}
+
+FSlateColor SFileTableRow::GetDisplayColor() const
+{
+	return TreeItem->GetAssetTypeColor();
+}
+
+void SFileTableRow::OnDragEnter(FGeometry const& InGeometry, FDragDropEvent const& InDragDropEvent)
+{
+	TSharedPtr<FDragDropOperation> DragOperation = InDragDropEvent.GetOperation();
+	DragOperation->SetCursorOverride(EMouseCursor::SlashedCircle);
+}
+
+void SFileTableRow::OnDragLeave(FDragDropEvent const& InDragDropEvent)
+{
+	TSharedPtr<FDragDropOperation> DragOperation = InDragDropEvent.GetOperation();
+	DragOperation->SetCursorOverride(EMouseCursor::None);
+}
+
+
+void SShelvedFilesTableRow::Construct(const FArguments& InArgs, const TSharedRef<STableViewBase>& InOwnerTableView)
+{
+	STableRow<FChangelistTreeItemPtr>::Construct(
+		STableRow<FChangelistTreeItemPtr>::FArguments()
+		.Content()
+		[
+			SNew(SHorizontalBox)
+				+SHorizontalBox::Slot()
+				.AutoWidth()
+				.VAlign(VAlign_Center)
+				.Padding(5, 0, 0, 0)
+				[
+					SNew(SImage)
+					.Image(InArgs._Icon)
+				]
+				+SHorizontalBox::Slot()
+				.Padding(2.0f, 1.0f)
+				.VAlign(VAlign_Center)
+				[
+					SNew(STextBlock)
+					.Text(InArgs._Text)
+				]
+		],
+		InOwnerTableView);
+}
+
+
+void SOfflineFileTableRow::Construct(const FArguments& InArgs, const TSharedRef<STableViewBase>& InOwner)
+{
+	TreeItem = static_cast<FOfflineFileTreeItem*>(InArgs._TreeItemToVisualize.Get());
+
+	FSuperRowType::FArguments Args = FSuperRowType::FArguments().ShowSelection(true);
+	FSuperRowType::Construct(Args, InOwner);
+}
+
+TSharedRef<SWidget> SOfflineFileTableRow::GenerateWidgetForColumn(const FName& ColumnId)
+{
+	if (ColumnId == SourceControlFileViewColumnId::Icon)
+	{
+		return SNew(SBox)
+			.WidthOverride(16) // Small Icons are usually 16x16
+			.HAlign(HAlign_Center)
+			[
+				SNew(SImage)
+				.Image(FAppStyle::GetBrush(FName("SourceControl.OfflineFile_Small")))
+			];
+	}
+	else if (ColumnId == SourceControlFileViewColumnId::Name)
+	{
+		return SNew(STextBlock)
+			.Text(this, &SOfflineFileTableRow::GetDisplayName);
+	}
+	else if (ColumnId == SourceControlFileViewColumnId::Path)
+	{
+		return SNew(STextBlock)
+			.Text(this, &SOfflineFileTableRow::GetDisplayPath)
+			.ToolTipText(this, &SOfflineFileTableRow::GetFilename);
+	}
+	else if (ColumnId == SourceControlFileViewColumnId::Type)
+	{
+		return SNew(STextBlock)
+			.Text(this, &SOfflineFileTableRow::GetDisplayType)
+			.ColorAndOpacity(this, &SOfflineFileTableRow::GetDisplayColor);
+	}
+	else
+	{
+		return SNullWidget::NullWidget;
+	}
+}
+
+FText SOfflineFileTableRow::GetDisplayName() const
+{
+	return TreeItem->GetDisplayName();
+}
+
+FText SOfflineFileTableRow::GetFilename() const
+{
+	return TreeItem->GetPackageName();
+}
+
+FText SOfflineFileTableRow::GetDisplayPath() const
+{
+	return TreeItem->GetDisplayPath();
+}
+
+FText SOfflineFileTableRow::GetDisplayType() const
+{
+	return TreeItem->GetDisplayType();
+}
+
+FSlateColor SOfflineFileTableRow::GetDisplayColor() const
+{
+	return TreeItem->GetDisplayColor();
+}
+
+#undef LOCTEXT_NAMESPACE
diff --git a/Engine/Source/Editor/SourceControlWindows/Private/SSourceControlChangelistRows.h b/Engine/Source/Editor/SourceControlWindows/Private/SSourceControlChangelistRows.h
new file mode 100644
--- /dev/null
+++ b/Engine/Source/Editor/SourceControlWindows/Private/SSourceControlChangelistRows.h
@@ -0,0 +1,188 @@
+// Copyright Epic Games, Inc. All Rights Reserved.
+
+#pragma once
+
+#include "CoreMinimal.h"
+
+#include "Widgets/SCompoundWidget.h"
+#include "Widgets/Views/STableRow.h"
+#include "Input/DragAndDrop.h"
+#include "SSourceControlCommon.h"
+
+/** Lists the unique column IDs used in the list view displaying controlled/uncontrolled changelist files. */
+namespace SourceControlFileViewColumnId
+{
+	/** The icon column id. */
+	extern const FName Icon;
+	/** The file/asset name column id. */
+	extern const FName Name;
+	/** The file/asset path column id. */
+	extern const FName Path;
+	/** The file/asset type column Id. */
+	extern const FName Type;
+}
+
+
+/** Displays a changed list row (icon, cl number, description) */
+class SChangelistTableRow : public STableRow<TSharedPtr<IChangelistTreeItem>>
+{
+public:
+	SLATE_BEGIN_ARGS(SChangelistTableRow)
+		: _TreeItemToVisualize()
+		, _OnPostDrop()
+	{
+	}
+		SLATE_ARGUMENT(FChangelistTreeItemPtr, TreeItemToVisualize)
+		SLATE_EVENT(FSimpleDelegate, OnPostDrop)
+	SLATE_END_ARGS()
+
+public:
+	/**
+	* Construct child widgets that comprise this widget.
+	*
+	* @param InArgs Declaration from which to construct this widget.
+	*/
+	void Construct(const FArguments& InArgs, const TSharedRef<STableViewBase>& InOwner);
+
+	FText GetChangelistText() const;
+	FText GetChangelistDescriptionText() const;
+
+protected:
+	//~ Begin STableRow Interface.
+	virtual FReply OnDrop(const FGeometry& InGeometry, const FDragDropEvent& InDragDropEvent) override;
+
+private:
+	/** The info about the widget that we are visualizing. */
+	FChangelistTreeItem* TreeItem;
+
+	/** Delegate invoked once the drag and drop operation finished. */
+	FSimpleDelegate OnPostDrop;
+};
+
+
+/** Displays an uncontrolled changed list (icon, cl name, description) */
+class SUncontrolledChangelistTableRow : public STableRow<FChangelistTreeItemPtr>
+{
+public:
+	SLATE_BEGIN_ARGS(SUncontrolledChangelistTableRow)
+		: _TreeItemToVisualize()
+		, _OnPostDrop()
+	{
+	}
+		SLATE_ARGUMENT(FChangelistTreeItemPtr, TreeItemToVisualize)
+		SLATE_EVENT(FSimpleDelegate, OnPostDrop)
+	SLATE_END_ARGS()
+
+public:
+	/**
+	* Construct child widgets that comprise this widget.
+	*
+	* @param InArgs Declaration from which to construct this widget.
+	*/
+	void Construct(const FArguments& InArgs, const TSharedRef<STableViewBase>& InOwner);
+
+	FText GetChangelistText() const;
+	FText GetChangelistDescriptionText() const;
+
+protected:
+	//~ Begin STableRow Interface.
+	virtual FReply OnDrop(const FGeometry& InGeometry, const FDragDropEvent& InDragDropEvent) override;
+
+private:
+	/** The info about the widget that we are visualizing. */
+	FUncontrolledChangelistTreeItem* TreeItem;
+
+	/** Invoked once a drag and drop operation completes. */
+	FSimpleDelegate OnPostDrop;
+};
+
+
+/** Display information about a file (icon, name, location, type, etc.) */
+class SFileTableRow : public SMultiColumnTableRow<FChangelistTreeItemPtr>
+{
+public:
+	SLATE_BEGIN_ARGS(SFileTableRow)
+		: _TreeItemToVisualize()
+	{
+	}
+		SLATE_ARGUMENT(FChangelistTreeItemPtr, TreeItemToVisualize)
+		SLATE_EVENT(FOnDragDetected, OnDragDetected)
+	SLATE_END_ARGS()
+
+public:
+	/**
+	* Construct child widgets that comprise this widget.
+	*
+	* @param InArgs Declaration from which to construct this widget.
+	*/
+	void Construct(const FArguments& InArgs, const TSharedRef<STableViewBase>& InOwner);
+
+	// SMultiColumnTableRow overrides
+	virtual TSharedRef<SWidget> GenerateWidgetForColumn(const FName& ColumnName) override;
+
+	FText GetDisplayName() const;
+	FText GetFilename() const;
+	FText GetDisplayPath() const;
+	FText GetDisplayType() const;
+	FSlateColor GetDisplayColor() const;
+
+protected:
+	//~ Begin STableRow Interface.
+	virtual void OnDragEnter(FGeometry const& InGeometry, FDragDropEvent const& InDragDropEvent) override;
+	virtual void OnDragLeave(FDragDropEvent const& InDragDropEvent) override;
+	//~ End STableRow Interface.
+
+private:
+	/** The info about the widget that we are visualizing. */
+	FFileTreeItem* TreeItem;
+};
+
+
+/** Display the shelved files group node. It displays 'Shelved Files (x)' where X is the nubmer of file shelved. */
+class SShelvedFilesTableRow : public STableRow<TSharedPtr<IChangelistTreeItem>>
+{
+public:
+	SLATE_BEGIN_ARGS(SShelvedFilesTableRow)
+		: _Icon(nullptr)
+	{
+	}
+		SLATE_ARGUMENT(const FSlateBrush*, Icon)
+		SLATE_ARGUMENT(FText, Text)
+	SLATE_END_ARGS()
+
+	void Construct(const FArguments& InArgs, const TSharedRef<STableViewBase>& InOwnerTableView);
+};
+
+
+/** Display information about an offline file (icon, name, location, type, etc.). */
+class SOfflineFileTableRow : public SMultiColumnTableRow<FChangelistTreeItemPtr>
+{
+public:
+	SLATE_BEGIN_ARGS(SOfflineFileTableRow)
+		: _TreeItemToVisualize()
+	{
+	}
+		SLATE_ARGUMENT(FChangelistTreeItemPtr, TreeItemToVisualize)
+	SLATE_END_ARGS()
+
+public:
+	/**
+	* Construct child widgets that comprise this widget.
+	*
+	* @param InArgs Declaration from which to construct this widget.
+	*/
+	void Construct(const FArguments& InArgs, const TSharedRef<STableViewBase>& InOwner);
+
+	// SMultiColumnTableRow overrides
+	virtual TSharedRef<SWidget> GenerateWidgetForColumn(const FName& ColumnName) override;
+
+	FText GetDisplayName() const;
+	FText GetFilename() const;
+	FText GetDisplayPath() const;
+	FText GetDisplayType() const;
+	FSlateColor GetDisplayColor() const;
+
+private:
+	/** The info about the widget that we are visualizing. */
+	FOfflineFileTreeItem* TreeItem;
+};
diff --git a/Engine/Source/Editor/SourceControlWindows/Private/SSourceControlChangelists.cpp b/Engine/Source/Editor/SourceControlWindows/Private/SSourceControlChangelists.cpp
--- a/Engine/Source/Editor/SourceControlWindows/Private/SSourceControlChangelists.cpp
+++ b/Engine/Source/Editor/SourceControlWindows/Private/SSourceControlChangelists.cpp
@@ -16,28 +16,23 @@
 #include "Widgets/Layout/SBox.h"
 #include "Widgets/Layout/SSplitter.h"
 #include "Widgets/Layout/SExpandableArea.h"
-#include "Widgets/Notifications/SNotificationList.h"
 
 #include "ISourceControlProvider.h"
 #include "ISourceControlModule.h"
 #include "ISourceControlWindowsModule.h"
 #include "UncontrolledChangelistsModule.h"
 #include "SourceControlOperations.h"
-#include "ToolMenus.h"
-#include "SFileTableRow.h"
+#include "SSourceControlChangelistRows.h"
 #include "SSourceControlDescription.h"
 #include "SourceControlWindows.h"
-#include "SSourceControlFileDialog.h"
 #include "SourceControlHelpers.h"
 #include "SourceControlPreferences.h"
 #include "SourceControlMenuContext.h"
 #include "AssetToolsModule.h"
-#include "ContentBrowserModule.h"
-#include "IContentBrowserSingleton.h"
 #include "Misc/MessageDialog.h"
 #include "Misc/ScopedSlowTask.h"
-#include "Algo/AnyOf.h"
 #include "HAL/PlatformTime.h"
+#include "ToolMenus.h"
 
 #include "SSourceControlSubmit.h"
 #include "Framework/Application/SlateApplication.h"
@@ -49,46 +44,6 @@
 namespace
 {
 
-/** Wraps the execution of a changelist operations with a slow task. */
-void ExecuteChangelistOperationWithSlowTaskWrapper(const FText& Message, const TFunction<void()>& ChangelistTask)
-{
-	// NOTE: This is a ugly workaround for P4 because the generic popup feedback operations in FScopedSourceControlProgress() was supressed for all synchrounous
-	//       operations. For other source control providers, the popup still shows up and showing a slow task and the FScopedSourceControlProgress at the same
-	//       time is a bad user experience. Until we fix source control popup situation in general in the Editor, this hack is in place to avoid the double popup.
-	//       At the time of writing, the other source control provider that supports changelists is Plastic.
-	if (ISourceControlModule::Get().GetProvider().GetName() == "Perforce")
-	{
-		FScopedSlowTask Progress(0.f, Message);
-		Progress.MakeDialog();
-		ChangelistTask();
-	}
-	else
-	{
-		ChangelistTask();
-	}
-}
-
-/** Wraps the execution of an uncontrolled changelist operations with a slow task. */
-void ExecuteUncontrolledChangelistOperationWithSlowTaskWrapper(const FText& Message, const TFunction<void()>& UncontrolledChangelistTask)
-{
-	ExecuteChangelistOperationWithSlowTaskWrapper(Message, UncontrolledChangelistTask);
-}
-
-/** Displays toast notification to report the status of task. */
-void DisplaySourceControlOperationNotification(const FText& Message, SNotificationItem::ECompletionState CompletionState)
-{
-	if (Message.IsEmpty())
-	{
-		return;
-	}
-
-	FNotificationInfo NotificationInfo(Message);
-	NotificationInfo.ExpireDuration = 6.0f;
-	NotificationInfo.Hyperlink = FSimpleDelegate::CreateLambda([]() { FGlobalTabmanager::Get()->TryInvokeTab(FName("OutputLog")); });
-	NotificationInfo.HyperlinkText = LOCTEXT("ShowOutputLogHyperlink", "Show Output Log");
-	FSlateNotificationManager::Get().AddNotification(NotificationInfo)->SetCompletionState(CompletionState);
-}
-
 /** Returns true if a source control provider is enable and support changeslists. */
 bool AreControlledChangelistsEnabled()
 {
@@ -143,56 +98,8 @@ FText UpdateChangelistDescriptionToSubmitIfNeeded(const bool bInValidationResult
 	return InChangelistDescription;
 }
 
-bool OpenConflictDialog(const TArray<FSourceControlStateRef>& InFilesConflicts)
-{
-	TSharedPtr<SWindow> Window;
-	TSharedPtr<SSourceControlFileDialog> SourceControlFileDialog;
-
-	Window = SNew(SWindow)
-			 .Title(LOCTEXT("CheckoutPackagesDialogTitle", "Check Out Assets"))
-			 .SizingRule(ESizingRule::UserSized)
-			 .ClientSize(FVector2D(1024.0f, 512.0f))
-			 .SupportsMaximize(false)
-			 .SupportsMinimize(false)
-			 [
-			 	SNew(SBorder)
-			 	.Padding(4.f)
-			 	.BorderImage(FAppStyle::GetBrush("ToolPanel.GroupBorder"))
-			 	[
-			 		SAssignNew(SourceControlFileDialog, SSourceControlFileDialog)
-			 		.Message(LOCTEXT("CheckoutPackagesDialogMessage", "Conflict detected in the following assets:"))
-			 		.Warning(LOCTEXT("CheckoutPackagesWarnMessage", "Warning: These assets are locked or not at the head revision. You may lose your changes if you continue, as you will be unable to submit them to source control."))
-			 		.Files(InFilesConflicts)
-			 	]
-			 ];
-
-	SourceControlFileDialog->SetWindow(Window);
-	Window->SetWidgetToFocusOnActivate(SourceControlFileDialog);
-	GEditor->EditorAddModalWindow(Window.ToSharedRef());
-
-	return SourceControlFileDialog->IsProceedButtonPressed();
-}
-
 } // Anonymous namespace
 
-/** Implements drag and drop operation. */
-struct FSCCFileDragDropOp : public FDragDropOperation
-{
-	DRAG_DROP_OPERATOR_TYPE(FSCCFileDragDropOp, FDragDropOperation);
-
-	using FDragDropOperation::Construct;
-
-	virtual TSharedPtr<SWidget> GetDefaultDecorator() const override
-	{
-		FSourceControlStateRef FileState = Files.IsEmpty() ? UncontrolledFiles[0] : Files[0];
-
-		return SSourceControlCommon::GetSCCFileWidget(MoveTemp(FileState));
-	}
-
-	TArray<FSourceControlStateRef> Files;
-	TArray<FSourceControlStateRef> UncontrolledFiles;
-};
-
 
 DECLARE_DELEGATE(FOnSearchBoxExpanded)
 
@@ -396,7 +303,7 @@ void SSourceControlChangelistsWidget::Construct(const FArguments& InArgs)
 	SourceControlStateChangedDelegateHandle = SCCModule.GetProvider().RegisterSourceControlStateChanged_Handle(FSourceControlStateChanged::FDelegate::CreateSP(this, &SSourceControlChangelistsWidget::OnSourceControlStateChanged));
 	UncontrolledChangelistModule.OnUncontrolledChangelistModuleChanged.AddSP(this, &SSourceControlChangelistsWidget::OnSourceControlStateChanged);
 
-	PrimarySortedColumn = "Name";
+	PrimarySortedColumn = SourceControlFileViewColumnId::Name;
 
 	ChangelistTreeView = CreateChangelistTreeView(ChangelistTreeNodes);
 	UncontrolledChangelistTreeView = CreateChangelistTreeView(UncontrolledChangelistTreeNodes);
@@ -568,11 +475,11 @@ void SSourceControlChangelistsWidget::EditChangelistDescription(const FText& InN
 		{
 			if (InResult == ECommandResult::Succeeded)
 			{
-				DisplaySourceControlOperationNotification(LOCTEXT("Update_Changelist_Description_Succeeded", "Changelist description successfully updated."), SNotificationItem::CS_Success);
+				SSourceControlCommon::DisplaySourceControlOperationNotification(LOCTEXT("Update_Changelist_Description_Succeeded", "Changelist description successfully updated."), SNotificationItem::CS_Success);
 			}
 			else if (InResult == ECommandResult::Failed)
 			{
-				DisplaySourceControlOperationNotification(LOCTEXT("Update_Changelist_Description_Failed", "Failed to update changelist description."), SNotificationItem::CS_Fail);
+				SSourceControlCommon::DisplaySourceControlOperationNotification(LOCTEXT("Update_Changelist_Description_Failed", "Failed to update changelist description."), SNotificationItem::CS_Fail);
 			}
 		}));
 }
@@ -1188,7 +1095,7 @@ void SSourceControlChangelistsWidget::Execute(const FText& Message, const TShare
 	}
 	else
 	{
-		ExecuteChangelistOperationWithSlowTaskWrapper(Message, [&]()
+		SSourceControlCommon::ExecuteChangelistOperationWithSlowTaskWrapper(Message, [&]()
 		{
 			ECommandResult::Type Result = SourceControlProvider.Execute(InOperation, InChangelist, InFiles, InConcurrency, InOperationCompleteDelegate);
 			OnEndSourceControlOperation(InOperation, Result);
@@ -1198,7 +1105,7 @@ void SSourceControlChangelistsWidget::Execute(const FText& Message, const TShare
 
 void SSourceControlChangelistsWidget::ExecuteUncontrolledChangelistOperation(const FText& Message, const TFunction<void()>& UncontrolledOperation)
 {
-	ExecuteUncontrolledChangelistOperationWithSlowTaskWrapper(Message, UncontrolledOperation);
+	SSourceControlCommon::ExecuteUncontrolledChangelistOperationWithSlowTaskWrapper(Message, UncontrolledOperation);
 }
 
 void SSourceControlChangelistsWidget::OnStartSourceControlOperation(TSharedRef<ISourceControlOperation> Operation, const FText& Message)
@@ -1269,11 +1176,11 @@ void SSourceControlChangelistsWidget::OnNewChangelist()
 			{
 				if (InResult == ECommandResult::Succeeded)
 				{
-					DisplaySourceControlOperationNotification(LOCTEXT("Create_Changelist_Succeeded", "Changelist successfully created."), SNotificationItem::CS_Success);
+					SSourceControlCommon::DisplaySourceControlOperationNotification(LOCTEXT("Create_Changelist_Succeeded", "Changelist successfully created."), SNotificationItem::CS_Success);
 				}
 				else if (InResult == ECommandResult::Failed)
 				{
-					DisplaySourceControlOperationNotification(LOCTEXT("Create_Changelist_Failed", "Failed to create the changelist."), SNotificationItem::CS_Fail);
+					SSourceControlCommon::DisplaySourceControlOperationNotification(LOCTEXT("Create_Changelist_Failed", "Failed to create the changelist."), SNotificationItem::CS_Fail);
 				}
 			}));
 }
@@ -1292,11 +1199,11 @@ void SSourceControlChangelistsWidget::OnDeleteChangelist()
 		{
 			if (InResult == ECommandResult::Succeeded)
 			{
-				DisplaySourceControlOperationNotification(LOCTEXT("Delete_Changelist_Succeeded", "Changelist successfully deleted."), SNotificationItem::CS_Success);
+				SSourceControlCommon::DisplaySourceControlOperationNotification(LOCTEXT("Delete_Changelist_Succeeded", "Changelist successfully deleted."), SNotificationItem::CS_Success);
 			}
 			else if (InResult == ECommandResult::Failed)
 			{
-				DisplaySourceControlOperationNotification(LOCTEXT("Delete_Changelist_Failed", "Failed to delete the selected changelist."), SNotificationItem::CS_Fail);
+				SSourceControlCommon::DisplaySourceControlOperationNotification(LOCTEXT("Delete_Changelist_Failed", "Failed to delete the selected changelist."), SNotificationItem::CS_Fail);
 			}
 		}));
 }
@@ -1371,11 +1278,11 @@ void SSourceControlChangelistsWidget::OnRevertUnchanged()
 			// NOTE: This operation message should tell how many files were reverted and how many weren't.
 			if (Operation->GetResultInfo().ErrorMessages.Num() == 0)
 			{
-				DisplaySourceControlOperationNotification(LOCTEXT("Revert_Unchanged_Files_Succeeded", "Unchanged files were reverted."), SNotificationItem::CS_Success);
+				SSourceControlCommon::DisplaySourceControlOperationNotification(LOCTEXT("Revert_Unchanged_Files_Succeeded", "Unchanged files were reverted."), SNotificationItem::CS_Success);
 			}
 			else if (InType == ECommandResult::Failed)
 			{
-				DisplaySourceControlOperationNotification(LOCTEXT("Revert_Unchanged_Files_Failed", "Failed to revert unchanged files."), SNotificationItem::CS_Fail);
+				SSourceControlCommon::DisplaySourceControlOperationNotification(LOCTEXT("Revert_Unchanged_Files_Failed", "Failed to revert unchanged files."), SNotificationItem::CS_Fail);
 			}
 		}));
 }
@@ -1434,15 +1341,15 @@ void SSourceControlChangelistsWidget::OnRevert()
 
 		if (!SelectedControlledFiles.IsEmpty())
 		{
-			ExecuteChangelistOperationWithSlowTaskWrapper(LOCTEXT("Reverting_Files", "Reverting file(s)..."), [&SelectedControlledFiles]()
+			SSourceControlCommon::ExecuteChangelistOperationWithSlowTaskWrapper(LOCTEXT("Reverting_Files", "Reverting file(s)..."), [&SelectedControlledFiles]()
 			{
 				if (SourceControlHelpers::RevertAndReloadPackages(SelectedControlledFiles))
 				{
-					DisplaySourceControlOperationNotification(LOCTEXT("Revert_Files_Succeeded", "The selected file(s) were reverted."), SNotificationItem::CS_Success);
+					SSourceControlCommon::DisplaySourceControlOperationNotification(LOCTEXT("Revert_Files_Succeeded", "The selected file(s) were reverted."), SNotificationItem::CS_Success);
 				}
 				else
 				{
-					DisplaySourceControlOperationNotification(LOCTEXT("Revert_Files_Failed", "Failed to revert the selected file(s)."), SNotificationItem::CS_Fail);
+					SSourceControlCommon::DisplaySourceControlOperationNotification(LOCTEXT("Revert_Files_Failed", "Failed to revert the selected file(s)."), SNotificationItem::CS_Fail);
 				}
 			});
 		}
@@ -1510,11 +1417,11 @@ void SSourceControlChangelistsWidget::OnShelve()
 		{
 			if (InResult == ECommandResult::Succeeded)
 			{
-				DisplaySourceControlOperationNotification(LOCTEXT("Shelve_Files_Succeeded", "The selected file(s) were shelved."), SNotificationItem::CS_Success);
+				SSourceControlCommon::DisplaySourceControlOperationNotification(LOCTEXT("Shelve_Files_Succeeded", "The selected file(s) were shelved."), SNotificationItem::CS_Success);
 			}
 			else if (InResult == ECommandResult::Failed)
 			{
-				DisplaySourceControlOperationNotification(LOCTEXT("Shelve_Files_Failed", "Failed to shelved the selected file(s)."), SNotificationItem::CS_Fail);
+				SSourceControlCommon::DisplaySourceControlOperationNotification(LOCTEXT("Shelve_Files_Failed", "Failed to shelved the selected file(s)."), SNotificationItem::CS_Fail);
 			}
 		}));
 }
@@ -1527,11 +1434,11 @@ void SSourceControlChangelistsWidget::OnUnshelve()
 		{
 			if (InResult == ECommandResult::Succeeded)
 			{
-				DisplaySourceControlOperationNotification(LOCTEXT("Unshelve_Files_Succeeded", "The selected file(s) were unshelved."), SNotificationItem::CS_Success);
+				SSourceControlCommon::DisplaySourceControlOperationNotification(LOCTEXT("Unshelve_Files_Succeeded", "The selected file(s) were unshelved."), SNotificationItem::CS_Success);
 			}
 			else if (InResult == ECommandResult::Failed)
 			{
-				DisplaySourceControlOperationNotification(LOCTEXT("Unshelve_Files_Failed", "Failed to unshelved the selected file(s)."), SNotificationItem::CS_Fail);
+				SSourceControlCommon::DisplaySourceControlOperationNotification(LOCTEXT("Unshelve_Files_Failed", "Failed to unshelved the selected file(s)."), SNotificationItem::CS_Fail);
 			}
 		}));
 }
@@ -1544,11 +1451,11 @@ void SSourceControlChangelistsWidget::OnDeleteShelvedFiles()
 		{
 			if (InResult == ECommandResult::Succeeded)
 			{
-				DisplaySourceControlOperationNotification(LOCTEXT("Delete_Shelved_Files_Succeeded", "The selected shelved file(s) were deleted."), SNotificationItem::CS_Success);
+				SSourceControlCommon::DisplaySourceControlOperationNotification(LOCTEXT("Delete_Shelved_Files_Succeeded", "The selected shelved file(s) were deleted."), SNotificationItem::CS_Success);
 			}
 			else if (InResult == ECommandResult::Failed)
 			{
-				DisplaySourceControlOperationNotification(LOCTEXT("Delete_Shelved_Files_Failed", "Failed to delete the selected shelved file(s)."), SNotificationItem::CS_Fail);
+				SSourceControlCommon::DisplaySourceControlOperationNotification(LOCTEXT("Delete_Shelved_Files_Failed", "Failed to delete the selected shelved file(s)."), SNotificationItem::CS_Fail);
 			}
 		}));
 }
@@ -1752,12 +1659,12 @@ void SSourceControlChangelistsWidget::OnSubmitChangelist()
 				{
 					if (InResult == ECommandResult::Succeeded)
 					{
-						DisplaySourceControlOperationNotification(SubmitChangelistOperation->GetSuccessMessage(), SNotificationItem::CS_Success);
+						SSourceControlCommon::DisplaySourceControlOperationNotification(SubmitChangelistOperation->GetSuccessMessage(), SNotificationItem::CS_Success);
 						bCheckinSuccess = true;
 					}
 					else if (InResult == ECommandResult::Failed)
 					{
-						DisplaySourceControlOperationNotification(LOCTEXT("SCC_Checkin_Failed", "Failed to check in files!"), SNotificationItem::CS_Fail);
+						SSourceControlCommon::DisplaySourceControlOperationNotification(LOCTEXT("SCC_Checkin_Failed", "Failed to check in files!"), SNotificationItem::CS_Fail);
 					}
 				}));
 		}
@@ -1933,14 +1840,14 @@ void SSourceControlChangelistsWidget::OnMoveFiles()
 					// NOTE: Perform uncontrolled move only if the new changelist was created and the controlled file were move.
 					if ((!SelectedUncontrolledFiles.IsEmpty()) && static_cast<FNewChangelist&>(Operation.Get()).GetNewChangelist().IsValid())
 					{
-						FUncontrolledChangelistsModule::Get().MoveFilesToControlledChangelist(SelectedUncontrolledFiles, static_cast<FNewChangelist&>(Operation.Get()).GetNewChangelist(), OpenConflictDialog);
+						FUncontrolledChangelistsModule::Get().MoveFilesToControlledChangelist(SelectedUncontrolledFiles, static_cast<FNewChangelist&>(Operation.Get()).GetNewChangelist(), SSourceControlCommon::OpenConflictDialog);
 					}
 
-					DisplaySourceControlOperationNotification(LOCTEXT("Move_Files_New_Changelist_Succeeded", "Files were successfully moved to a new changelist."), SNotificationItem::CS_Success);
+					SSourceControlCommon::DisplaySourceControlOperationNotification(LOCTEXT("Move_Files_New_Changelist_Succeeded", "Files were successfully moved to a new changelist."), SNotificationItem::CS_Success);
 				}
 				if (InResult == ECommandResult::Failed)
 				{
-					DisplaySourceControlOperationNotification(LOCTEXT("Move_Files_New_Changelist_Failed", "Failed to move the file to the new changelist."), SNotificationItem::CS_Fail);
+					SSourceControlCommon::DisplaySourceControlOperationNotification(LOCTEXT("Move_Files_New_Changelist_Failed", "Failed to move the file to the new changelist."), SNotificationItem::CS_Fail);
 				}
 			}));
 	}
@@ -1974,14 +1881,14 @@ void SSourceControlChangelistsWidget::OnMoveFiles()
 							// Perform an uncontrolled move only if the controlled file were move successfully.
 							if (!SelectedUncontrolledFiles.IsEmpty())
 							{
-								FUncontrolledChangelistsModule::Get().MoveFilesToControlledChangelist(SelectedUncontrolledFiles, Changelist, OpenConflictDialog);
+								FUncontrolledChangelistsModule::Get().MoveFilesToControlledChangelist(SelectedUncontrolledFiles, Changelist, SSourceControlCommon::OpenConflictDialog);
 							}
 
-							DisplaySourceControlOperationNotification(LOCTEXT("Move_Files_Between_Changelist_Succeeded", "File(s) successfully moved to the selected changelist."), SNotificationItem::CS_Success);
+							SSourceControlCommon::DisplaySourceControlOperationNotification(LOCTEXT("Move_Files_Between_Changelist_Succeeded", "File(s) successfully moved to the selected changelist."), SNotificationItem::CS_Success);
 						}
 						else if (InResult == ECommandResult::Failed)
 						{
-							DisplaySourceControlOperationNotification(LOCTEXT("Move_Files_Between_Changelist_Failed", "Failed to move the file(s) to the selected changelist."), SNotificationItem::CS_Fail);
+							SSourceControlCommon::DisplaySourceControlOperationNotification(LOCTEXT("Move_Files_Between_Changelist_Failed", "Failed to move the file(s) to the selected changelist."), SNotificationItem::CS_Fail);
 						}
 					}));
 			}
@@ -2259,420 +2166,36 @@ TSharedRef<STreeView<FChangelistTreeItemPtr>> SSourceControlChangelistsWidget::C
 		.HeaderRow
 		(
 			SNew(SHeaderRow)
-			+SHeaderRow::Column("Icon")
+			+SHeaderRow::Column(SourceControlFileViewColumnId::Icon)
 			.DefaultLabel(FText::GetEmpty())
 			.FillSized(18)
-			.SortPriority(this, &SSourceControlChangelistsWidget::GetColumnSortPriority, FName("Icon"))
-			.SortMode(this, &SSourceControlChangelistsWidget::GetColumnSortMode, FName("Icon"))
+			.SortPriority(this, &SSourceControlChangelistsWidget::GetColumnSortPriority, SourceControlFileViewColumnId::Icon)
+			.SortMode(this, &SSourceControlChangelistsWidget::GetColumnSortMode, SourceControlFileViewColumnId::Icon)
 			.OnSort(this, &SSourceControlChangelistsWidget::OnColumnSortModeChanged)
 
-			+SHeaderRow::Column("Name")
+			+SHeaderRow::Column(SourceControlFileViewColumnId::Name)
 			.DefaultLabel(LOCTEXT("Name", "Name"))
 			.FillWidth(0.2f)
-			.SortPriority(this, &SSourceControlChangelistsWidget::GetColumnSortPriority, FName("Name"))
-			.SortMode(this, &SSourceControlChangelistsWidget::GetColumnSortMode, FName("Name"))
+			.SortPriority(this, &SSourceControlChangelistsWidget::GetColumnSortPriority, SourceControlFileViewColumnId::Name)
+			.SortMode(this, &SSourceControlChangelistsWidget::GetColumnSortMode, SourceControlFileViewColumnId::Name)
 			.OnSort(this, &SSourceControlChangelistsWidget::OnColumnSortModeChanged)
 
-			+SHeaderRow::Column("Path")
+			+SHeaderRow::Column(SourceControlFileViewColumnId::Path)
 			.DefaultLabel(LOCTEXT("Path", "Path"))
 			.FillWidth(0.6f)
-			.SortPriority(this, &SSourceControlChangelistsWidget::GetColumnSortPriority, FName("Path"))
-			.SortMode(this, &SSourceControlChangelistsWidget::GetColumnSortMode, FName("Path"))
+			.SortPriority(this, &SSourceControlChangelistsWidget::GetColumnSortPriority, SourceControlFileViewColumnId::Path)
+			.SortMode(this, &SSourceControlChangelistsWidget::GetColumnSortMode, SourceControlFileViewColumnId::Path)
 			.OnSort(this, &SSourceControlChangelistsWidget::OnColumnSortModeChanged)
 
-			+SHeaderRow::Column("Type")
+			+SHeaderRow::Column(SourceControlFileViewColumnId::Type)
 			.DefaultLabel(LOCTEXT("Type", "Type"))
 			.FillWidth(0.2f)
-			.SortPriority(this, &SSourceControlChangelistsWidget::GetColumnSortPriority, FName("Type"))
-			.SortMode(this, &SSourceControlChangelistsWidget::GetColumnSortMode, FName("Type"))
+			.SortPriority(this, &SSourceControlChangelistsWidget::GetColumnSortPriority, SourceControlFileViewColumnId::Type)
+			.SortMode(this, &SSourceControlChangelistsWidget::GetColumnSortMode, SourceControlFileViewColumnId::Type)
 			.OnSort(this, &SSourceControlChangelistsWidget::OnColumnSortModeChanged)
 		);
 }
 
-/** Displays a changed list row (icon, cl number, description) */
-class SChangelistTableRow : public STableRow<TSharedPtr<IChangelistTreeItem>>
-{
-public:
-	SLATE_BEGIN_ARGS(SChangelistTableRow)
-		: _TreeItemToVisualize()
-		, _OnPostDrop()
-	{}
-	SLATE_ARGUMENT(FChangelistTreeItemPtr, TreeItemToVisualize)
-	SLATE_EVENT(FSimpleDelegate, OnPostDrop)
-	SLATE_END_ARGS()
-
-public:
-	/**
-	* Construct child widgets that comprise this widget.
-	*
-	* @param InArgs Declaration from which to construct this widget.
-	*/
-	void Construct(const FArguments& InArgs, const TSharedRef<STableViewBase>& InOwner)
-	{
-		TreeItem = static_cast<FChangelistTreeItem*>(InArgs._TreeItemToVisualize.Get());
-		OnPostDrop = InArgs._OnPostDrop;
-
-		const FSlateBrush* IconBrush = (TreeItem != nullptr) ?
-			FAppStyle::GetBrush(TreeItem->ChangelistState->GetSmallIconName()) :
-			FAppStyle::GetBrush("SourceControl.Changelist");
-
-		SetToolTipText(GetChangelistDescriptionText());
-
-		STableRow<FChangelistTreeItemPtr>::Construct(
-			STableRow<FChangelistTreeItemPtr>::FArguments()
-			.Style(FAppStyle::Get(), "TableView.Row")
-			.Content()
-			[
-				SNew(SHorizontalBox)
-				+SHorizontalBox::Slot() // Icon
-				.AutoWidth()
-				[
-					SNew(SImage)
-					.Image(IconBrush)
-				]
-				+SHorizontalBox::Slot() // Changelist number.
-				.Padding(2, 0, 0, 0)
-				.AutoWidth()
-				[
-					SNew(STextBlock)
-					.Text(this, &SChangelistTableRow::GetChangelistText)
-				]
-				+SHorizontalBox::Slot() // Files count.
-				.Padding(4, 0, 4, 0)
-				.AutoWidth()
-				[
-					SNew(STextBlock)
-					.Text(FText::Format(INVTEXT("({0})"), TreeItem->GetFileCount()))
-				]
-				+SHorizontalBox::Slot()
-				.Padding(2, 0, 0, 0)
-				.AutoWidth()
-				[
-					SNew(STextBlock)
-					.Text(this, &SChangelistTableRow::GetChangelistDescriptionText)
-				]
-			], InOwner);
-	}
-
-	FText GetChangelistText() const
-	{
-		return TreeItem->GetDisplayText();
-	}
-
-	FText GetChangelistDescriptionText() const
-	{
-		FString DescriptionString = TreeItem->GetDescriptionText().ToString();
-		// Here we'll both remove \r\n (when edited from the dialog) and \n (when we get it from the SCC)
-		DescriptionString.ReplaceInline(TEXT("\r"), TEXT(""));
-		DescriptionString.ReplaceInline(TEXT("\n"), TEXT(" "));
-		DescriptionString.TrimEndInline();
-		return FText::FromString(DescriptionString);
-	}
-
-protected:
-	//~ Begin STableRow Interface.
-	virtual FReply OnDrop(const FGeometry& InGeometry, const FDragDropEvent& InDragDropEvent) override
-	{
-		TSharedPtr<FSCCFileDragDropOp> DropOperation = InDragDropEvent.GetOperationAs<FSCCFileDragDropOp>();
-		if (DropOperation.IsValid())
-		{
-			FSourceControlChangelistPtr DestChangelist = TreeItem->ChangelistState->GetChangelist();
-			check(DestChangelist.IsValid());
-
-			// NOTE: The UI don't show 'source controlled files' and 'uncontrolled files' at the same time. User cannot select and drag/drop both file types at the same time.
-			if (!DropOperation->Files.IsEmpty())
-			{
-				TArray<FString> Files;
-				Algo::Transform(DropOperation->Files, Files, [](const FSourceControlStateRef& State) { return State->GetFilename(); });
-
-				ExecuteChangelistOperationWithSlowTaskWrapper(LOCTEXT("Dropping_Files_On_Changelist", "Moving file(s) to the selected changelist..."), [&]()
-				{
-					ISourceControlProvider& SourceControlProvider = ISourceControlModule::Get().GetProvider();
-					SourceControlProvider.Execute(ISourceControlOperation::Create<FMoveToChangelist>(), DestChangelist, Files, EConcurrency::Synchronous, FSourceControlOperationComplete::CreateLambda(
-						[](const TSharedRef<ISourceControlOperation>& Operation, ECommandResult::Type InResult)
-						{
-							if (InResult == ECommandResult::Succeeded)
-							{
-								DisplaySourceControlOperationNotification(LOCTEXT("Drop_Files_On_Changelist_Succeeded", "File(s) successfully moved to the selected changelist."), SNotificationItem::CS_Success);
-							}
-							else if (InResult == ECommandResult::Failed)
-							{
-								DisplaySourceControlOperationNotification(LOCTEXT("Drop_Files_On_Changelist_Failed", "Failed to move the file(s) to the selected changelist."), SNotificationItem::CS_Fail);
-							}
-						}));
-				});
-			}
-			else if (!DropOperation->UncontrolledFiles.IsEmpty())
-			{
-				// NOTE: This function does several operations that can fails but we don't get feedback.
-				ExecuteUncontrolledChangelistOperationWithSlowTaskWrapper(LOCTEXT("Dropping_Uncontrolled_Files_On_Changelist", "Moving uncontrolled file(s) to the selected changelist..."), 
-					[&DropOperation, &DestChangelist]()
-				{
-					FUncontrolledChangelistsModule::Get().MoveFilesToControlledChangelist(DropOperation->UncontrolledFiles, DestChangelist, OpenConflictDialog);
-					
-					// TODO: Fix MoveFilesToControlledChangelist() to report the possible errors and display a notification.
-				});
-
-				OnPostDrop.ExecuteIfBound();
-			}
-		}
-
-		return FReply::Handled();
-	}
-	//~ End STableRow Interface.
-
-private:
-	/** The info about the widget that we are visualizing. */
-	FChangelistTreeItem* TreeItem;
-
-	/** Delegate invoked once the drag and drop operation finished. */
-	FSimpleDelegate OnPostDrop;
-};
-
-/** Displays an uncontrolled changed list (icon, cl number, description) */
-class SUncontrolledChangelistTableRow : public STableRow<FChangelistTreeItemPtr>
-{
-public:
-	SLATE_BEGIN_ARGS(SUncontrolledChangelistTableRow)
-		: _TreeItemToVisualize()
-		, _OnPostDrop()
-	{
-	}
-	SLATE_ARGUMENT(FChangelistTreeItemPtr, TreeItemToVisualize)
-	SLATE_EVENT(FSimpleDelegate, OnPostDrop)
-	SLATE_END_ARGS()
-
-public:
-	/**
-	* Construct child widgets that comprise this widget.
-	*
-	* @param InArgs Declaration from which to construct this widget.
-	*/
-	void Construct(const FArguments& InArgs, const TSharedRef<STableViewBase>& InOwner)
-	{
-		TreeItem = static_cast<FUncontrolledChangelistTreeItem*>(InArgs._TreeItemToVisualize.Get());
-		OnPostDrop = InArgs._OnPostDrop;
-
-			const FSlateBrush* IconBrush = (TreeItem != nullptr) ?
-				FAppStyle::GetBrush(TreeItem->UncontrolledChangelistState->GetSmallIconName()) :
-				FAppStyle::GetBrush("SourceControl.Changelist");
-
-		SetToolTipText(GetChangelistDescriptionText());
-
-		STableRow<FChangelistTreeItemPtr>::Construct(
-			STableRow<FChangelistTreeItemPtr>::FArguments()
-			.Style(FAppStyle::Get(), "TableView.Row")
-			.Content()
-			[
-				SNew(SHorizontalBox)
-				+ SHorizontalBox::Slot()
-				.AutoWidth()
-				.VAlign(VAlign_Center)
-				[
-					SNew(SImage)
-					.Image(IconBrush)
-				]
-				+ SHorizontalBox::Slot()
-				.AutoWidth()
-				.Padding(2.0f, 0.0f)
-				.VAlign(VAlign_Center)
-				[
-					SNew(STextBlock)
-					.Text(this, &SUncontrolledChangelistTableRow::GetChangelistText)
-				]
-				+SHorizontalBox::Slot() // Files/Offline file count.
-				.Padding(4, 0, 4, 0)
-				.AutoWidth()
-				[
-					SNew(STextBlock)
-					.Text(FText::Format(INVTEXT("({0})"), TreeItem->GetFileCount() + TreeItem->GetOfflineFileCount()))
-				]
-				+ SHorizontalBox::Slot()
-				.AutoWidth()
-				.Padding(2.0f, 0.0f)
-				.VAlign(VAlign_Center)
-				[
-					SNew(STextBlock)
-					.Text(this, &SUncontrolledChangelistTableRow::GetChangelistDescriptionText)
-				]
-			], InOwner);
-	}
-
-	FText GetChangelistText() const
-	{
-		return TreeItem->GetDisplayText();
-	}
-
-	FText GetChangelistDescriptionText() const
-	{
-		FString DescriptionString = TreeItem->GetDescriptionText().ToString();
-		// Here we'll both remove \r\n (when edited from the dialog) and \n (when we get it from the SCC)
-		DescriptionString.ReplaceInline(TEXT("\r"), TEXT(""));
-		DescriptionString.ReplaceInline(TEXT("\n"), TEXT(" "));
-		DescriptionString.TrimEndInline();
-		return FText::FromString(DescriptionString);
-	}
-
-protected:
-	//~ Begin STableRow Interface.
-	virtual FReply OnDrop(const FGeometry& InGeometry, const FDragDropEvent& InDragDropEvent) override
-	{
-		TSharedPtr<FSCCFileDragDropOp> Operation = InDragDropEvent.GetOperationAs<FSCCFileDragDropOp>();
-		
-		if (Operation.IsValid())
-		{
-			ExecuteUncontrolledChangelistOperationWithSlowTaskWrapper(LOCTEXT("Drag_File_To_Uncontrolled_Changelist", "Moving file(s) to the selected uncontrolled changelists..."),
-				[this, &Operation]()
-			{
-				FUncontrolledChangelistsModule::Get().MoveFilesToUncontrolledChangelist(Operation->Files, Operation->UncontrolledFiles, TreeItem->UncontrolledChangelistState->Changelist);
-			});
-
-			OnPostDrop.ExecuteIfBound();
-		}
-
-		return FReply::Handled();
-	}
-	//~ End STableRow Interface.
-
-private:
-	/** The info about the widget that we are visualizing. */
-	FUncontrolledChangelistTreeItem* TreeItem;
-
-	/** Invoked once a drag and drop operation completes. */
-	FSimpleDelegate OnPostDrop;
-};
-
-/** Display the shelved files group node. It displays 'Shelved Files (x)' where X is the nubmer of file shelved. */
-class SShelvedFilesTableRow : public STableRow<TSharedPtr<IChangelistTreeItem>>
-{
-public:
-	SLATE_BEGIN_ARGS(SShelvedFilesTableRow)
-		: _Icon(nullptr)
-	{
-	}
-		SLATE_ARGUMENT(const FSlateBrush*, Icon)
-		SLATE_ARGUMENT(FText, Text)
-	SLATE_END_ARGS()
-
-	void Construct(const FArguments& InArgs, const TSharedRef<STableViewBase>& InOwnerTableView)
-	{
-		STableRow<FChangelistTreeItemPtr>::Construct(
-			STableRow<FChangelistTreeItemPtr>::FArguments()
-			.Content()
-			[
-				SNew(SHorizontalBox)
-					+SHorizontalBox::Slot()
-					.AutoWidth()
-					.VAlign(VAlign_Center)
-					.Padding(5, 0, 0, 0)
-					[
-						SNew(SImage)
-						.Image(InArgs._Icon)
-					]
-					+SHorizontalBox::Slot()
-					.Padding(2.0f, 1.0f)
-					.VAlign(VAlign_Center)
-					[
-						SNew(STextBlock)
-						.Text(InArgs._Text)
-					]
-			],
-			InOwnerTableView);
-	}
-};
-
-/** Display information about an offline file (icon, name, location, type, etc.). */
-class SOfflineFileTableRow : public SMultiColumnTableRow<FChangelistTreeItemPtr>
-{
-public:
-	SLATE_BEGIN_ARGS(SOfflineFileTableRow)
-		: _TreeItemToVisualize()
-	{
-	}
-	SLATE_ARGUMENT(FChangelistTreeItemPtr, TreeItemToVisualize)
-	SLATE_END_ARGS()
-
-public:
-	/**
-	* Construct child widgets that comprise this widget.
-	*
-	* @param InArgs Declaration from which to construct this widget.
-	*/
-	void Construct(const FArguments& InArgs, const TSharedRef<STableViewBase>& InOwner)
-	{
-		TreeItem = static_cast<FOfflineFileTreeItem*>(InArgs._TreeItemToVisualize.Get());
-
-		FSuperRowType::FArguments Args = FSuperRowType::FArguments().ShowSelection(true);
-		FSuperRowType::Construct(Args, InOwner);
-	}
-
-	// SMultiColumnTableRow overrides
-	virtual TSharedRef<SWidget> GenerateWidgetForColumn(const FName& ColumnName) override
-	{
-		if (ColumnName == TEXT("Icon"))
-		{
-			return SNew(SBox)
-				.WidthOverride(16) // Small Icons are usually 16x16
-				.HAlign(HAlign_Center)
-				[
-					SNew(SImage)
-					.Image(FAppStyle::GetBrush(FName("SourceControl.OfflineFile_Small")))
-				];
-		}
-		else if (ColumnName == TEXT("Name"))
-		{
-			return SNew(STextBlock)
-				.Text(this, &SOfflineFileTableRow::GetDisplayName);
-		}
-		else if (ColumnName == TEXT("Path"))
-		{
-			return SNew(STextBlock)
-				.Text(this, &SOfflineFileTableRow::GetDisplayPath)
-				.ToolTipText(this, &SOfflineFileTableRow::GetFilename);
-		}
-		else if (ColumnName == TEXT("Type"))
-		{
-			return SNew(STextBlock)
-				.Text(this, &SOfflineFileTableRow::GetDisplayType)
-				.ColorAndOpacity(this, &SOfflineFileTableRow::GetDisplayColor);
-		}
-		else
-		{
-			return SNullWidget::NullWidget;
-		}
-	}
-
-	FText GetDisplayName() const
-	{
-		return TreeItem->GetDisplayName();
-	}
-
-	FText GetFilename() const
-	{
-		return TreeItem->GetPackageName();
-	}
-
-	FText GetDisplayPath() const
-	{
-		return TreeItem->GetDisplayPath();
-	}
-
-	FText GetDisplayType() const
-	{
-		return TreeItem->GetDisplayType();
-	}
-
-	FSlateColor GetDisplayColor() const
-	{
-		return TreeItem->GetDisplayColor();
-	}
-
-private:
-	/** The info about the widget that we are visualizing. */
-	FOfflineFileTreeItem* TreeItem;
-};
-
-
 TSharedRef<ITableRow> SSourceControlChangelistsWidget::OnGenerateRow(TSharedPtr<IChangelistTreeItem> InTreeItem, const TSharedRef<STableViewBase>& OwnerTable)
 {
 	switch (InTreeItem->GetTreeItemType())
@@ -2976,21 +2499,21 @@ void SSourceControlChangelistsWidget::SortFileView()
 		}
 	};
 
-	auto Compare = [&](const TSharedPtr<IChangelistTreeItem>& Lhs, const TSharedPtr<IChangelistTreeItem>& Rhs, const FName& ColName, EColumnSortMode::Type SortMode) -> bool
+	auto Compare = [&](const TSharedPtr<IChangelistTreeItem>& Lhs, const TSharedPtr<IChangelistTreeItem>& Rhs, const FName& ColumnId, EColumnSortMode::Type SortMode) -> bool
 	{
-		if (ColName == "Icon")
+		if (ColumnId == SourceControlFileViewColumnId::Icon)
 		{
 			return SortMode == EColumnSortMode::Ascending ? OperatorLessIcon(Lhs.Get(), Rhs.Get()) : OperatorLessIcon(Rhs.Get(), Lhs.Get());
 		}
-		else if (ColName == "Name")
+		else if (ColumnId == SourceControlFileViewColumnId::Name)
 		{
 			return SortMode == EColumnSortMode::Ascending ? GetName(Lhs.Get()) < GetName(Rhs.Get()) : GetName(Lhs.Get()) > GetName(Rhs.Get());
 		}
-		else if (ColName == "Path")
+		else if (ColumnId == SourceControlFileViewColumnId::Path)
 		{
 			return SortMode == EColumnSortMode::Ascending ? GetPath(Lhs.Get()) < GetPath(Rhs.Get()) : GetPath(Lhs.Get()) > GetPath(Rhs.Get());
 		}
-		else if (ColName == "Type")
+		else if (ColumnId == SourceControlFileViewColumnId::Type)
 		{
 			return SortMode == EColumnSortMode::Ascending ? GetType(Lhs.Get()) < GetType(Rhs.Get()) : GetType(Lhs.Get()) > GetType(Rhs.Get());
 		}
diff --git a/Engine/Source/Editor/SourceControlWindows/Private/SSourceControlCommon.cpp b/Engine/Source/Editor/SourceControlWindows/Private/SSourceControlCommon.cpp
--- a/Engine/Source/Editor/SourceControlWindows/Private/SSourceControlCommon.cpp
+++ b/Engine/Source/Editor/SourceControlWindows/Private/SSourceControlCommon.cpp
@@ -12,11 +12,14 @@
 #include "ISourceControlModule.h"
 #include "SourceControlAssetDataCache.h"
 #include "SourceControlHelpers.h"
+#include "SSourceControlFileDialog.h"
 
 #include "Widgets/SOverlay.h"
 #include "Widgets/Images/SImage.h"
 #include "Widgets/Images/SLayeredImage.h"
 #include "Widgets/Layout/SBox.h"
+#include"Framework/Notifications/NotificationManager.h"
+#include "Misc/ScopedSlowTask.h"
 
 #define LOCTEXT_NAMESPACE "SourceControlChangelist"
 
@@ -345,6 +348,77 @@ FText GetDefaultMultipleAsset()
 	return LOCTEXT("SourceCOntrol_ManyAssetType", "Multiple Assets");
 }
 
+/** Wraps the execution of a changelist operations with a slow task. */
+void ExecuteChangelistOperationWithSlowTaskWrapper(const FText& Message, const TFunction<void()>& ChangelistTask)
+{
+	// NOTE: This is a ugly workaround for P4 because the generic popup feedback operations in FScopedSourceControlProgress() was supressed for all synchrounous
+	//       operations. For other source control providers, the popup still shows up and showing a slow task and the FScopedSourceControlProgress at the same
+	//       time is a bad user experience. Until we fix source control popup situation in general in the Editor, this hack is in place to avoid the double popup.
+	//       At the time of writing, the other source control provider that supports changelists is Plastic.
+	if (ISourceControlModule::Get().GetProvider().GetName() == "Perforce")
+	{
+		FScopedSlowTask Progress(0.f, Message);
+		Progress.MakeDialog();
+		ChangelistTask();
+	}
+	else
+	{
+		ChangelistTask();
+	}
+}
+
+/** Wraps the execution of an uncontrolled changelist operations with a slow task. */
+void ExecuteUncontrolledChangelistOperationWithSlowTaskWrapper(const FText& Message, const TFunction<void()>& UncontrolledChangelistTask)
+{
+	ExecuteChangelistOperationWithSlowTaskWrapper(Message, UncontrolledChangelistTask);
+}
+
+/** Displays toast notification to report the status of task. */
+void DisplaySourceControlOperationNotification(const FText& Message, SNotificationItem::ECompletionState CompletionState)
+{
+	if (Message.IsEmpty())
+	{
+		return;
+	}
+
+	FNotificationInfo NotificationInfo(Message);
+	NotificationInfo.ExpireDuration = 6.0f;
+	NotificationInfo.Hyperlink = FSimpleDelegate::CreateLambda([]() { FGlobalTabmanager::Get()->TryInvokeTab(FName("OutputLog")); });
+	NotificationInfo.HyperlinkText = LOCTEXT("ShowOutputLogHyperlink", "Show Output Log");
+	FSlateNotificationManager::Get().AddNotification(NotificationInfo)->SetCompletionState(CompletionState);
+}
+
+bool OpenConflictDialog(const TArray<FSourceControlStateRef>& InFilesConflicts)
+{
+	TSharedPtr<SWindow> Window;
+	TSharedPtr<SSourceControlFileDialog> SourceControlFileDialog;
+
+	Window = SNew(SWindow)
+			 .Title(LOCTEXT("CheckoutPackagesDialogTitle", "Check Out Assets"))
+			 .SizingRule(ESizingRule::UserSized)
+			 .ClientSize(FVector2D(1024.0f, 512.0f))
+			 .SupportsMaximize(false)
+			 .SupportsMinimize(false)
+			 [
+			 	SNew(SBorder)
+			 	.Padding(4.f)
+			 	.BorderImage(FAppStyle::GetBrush("ToolPanel.GroupBorder"))
+			 	[
+			 		SAssignNew(SourceControlFileDialog, SSourceControlFileDialog)
+			 		.Message(LOCTEXT("CheckoutPackagesDialogMessage", "Conflict detected in the following assets:"))
+			 		.Warning(LOCTEXT("CheckoutPackagesWarnMessage", "Warning: These assets are locked or not at the head revision. You may lose your changes if you continue, as you will be unable to submit them to source control."))
+			 		.Files(InFilesConflicts)
+			 	]
+			 ];
+
+	SourceControlFileDialog->SetWindow(Window);
+	Window->SetWidgetToFocusOnActivate(SourceControlFileDialog);
+	GEditor->EditorAddModalWindow(Window.ToSharedRef());
+
+	return SourceControlFileDialog->IsProceedButtonPressed();
+}
+
+
 } // end of namespace SSourceControlCommon
 
 #undef LOCTEXT_NAMESPACE
diff --git a/Engine/Source/Editor/SourceControlWindows/Private/SSourceControlCommon.h b/Engine/Source/Editor/SourceControlWindows/Private/SSourceControlCommon.h
--- a/Engine/Source/Editor/SourceControlWindows/Private/SSourceControlCommon.h
+++ b/Engine/Source/Editor/SourceControlWindows/Private/SSourceControlCommon.h
@@ -8,6 +8,9 @@
 #include "SourceControlAssetDataCache.h"
 #include "ISourceControlProvider.h"
 #include "UncontrolledChangelistState.h"
+#include "Input/DragAndDrop.h"
+#include "Widgets/Notifications/SNotificationList.h"
+#include "Templates/Function.h"
 
 struct FAssetData;
 
@@ -66,6 +69,7 @@ protected:
 	TreeItemType Type;
 };
 
+
 /** Displays a changelist icon/number/description. */
 struct FChangelistTreeItem : public IChangelistTreeItem
 {
@@ -98,6 +102,7 @@ struct FChangelistTreeItem : public IChangelistTreeItem
 	TSharedRef<ISourceControlChangelistState> ChangelistState;
 };
 
+
 /** Displays an uncontrolled changelist icon/number/description. */
 struct FUncontrolledChangelistTreeItem : public IChangelistTreeItem
 {
@@ -130,6 +135,7 @@ struct FUncontrolledChangelistTreeItem : public IChangelistTreeItem
 	FUncontrolledChangelistStateRef UncontrolledChangelistState;
 };
 
+
 /** Displays a set of files under a changelist or uncontrolled changelist. */
 struct FFileTreeItem : public IChangelistTreeItem
 {
@@ -230,6 +236,7 @@ private:
 	bool bAssetsUpToDate;
 };
 
+
 /** Root node to group shelved files as children. */
 struct FShelvedChangelistTreeItem : public IChangelistTreeItem
 {
@@ -237,6 +244,7 @@ struct FShelvedChangelistTreeItem : public IChangelistTreeItem
 	FText GetDisplayText() const;
 };
 
+
 struct FShelvedFileTreeItem : public FFileTreeItem
 {
 	explicit FShelvedFileTreeItem(FSourceControlStateRef InFileState, bool bBeautifyPaths = true)
@@ -245,21 +253,17 @@ struct FShelvedFileTreeItem : public FFileTreeItem
 	}
 };
 
+
 struct FOfflineFileTreeItem : public IChangelistTreeItem
 {
 	explicit FOfflineFileTreeItem(const FString& InFilename);
 
 public:
 	const FString& GetFilename() const { return Filename; }
-
 	const FText& GetPackageName() const { return PackageName; }
-
 	const FText& GetDisplayName() const { return AssetName; }
-
 	const FText& GetDisplayPath() const { return AssetPath; }
-
 	const FText& GetDisplayType() const { return AssetType; }
-
 	const FSlateColor& GetDisplayColor() const { return AssetTypeColor; }
 
 private:
@@ -272,6 +276,7 @@ private:
 	FSlateColor AssetTypeColor;
 };
 
+
 namespace SSourceControlCommon
 {
 	TSharedRef<SWidget> GetSCCFileWidget(FSourceControlStateRef InFileState, bool bIsShelvedFile = false);
@@ -279,9 +284,32 @@ namespace SSourceControlCommon
 	FText GetDefaultAssetType();
 	FText GetDefaultUnknownAssetType();
 	FText GetDefaultMultipleAsset();
+
+	void ExecuteChangelistOperationWithSlowTaskWrapper(const FText& Message, const TFunction<void()>& ChangelistTask);
+	void ExecuteUncontrolledChangelistOperationWithSlowTaskWrapper(const FText& Message, const TFunction<void()>& UncontrolledChangelistTask);
+	void DisplaySourceControlOperationNotification(const FText& Message, SNotificationItem::ECompletionState CompletionState);
+	bool OpenConflictDialog(const TArray<FSourceControlStateRef>& InFilesConflicts);
 }
 
 
+/** Implements drag and drop operation. */
+struct FSCCFileDragDropOp : public FDragDropOperation
+{
+	DRAG_DROP_OPERATOR_TYPE(FSCCFileDragDropOp, FDragDropOperation);
+
+	using FDragDropOperation::Construct;
+
+	virtual TSharedPtr<SWidget> GetDefaultDecorator() const override
+	{
+		FSourceControlStateRef FileState = Files.IsEmpty() ? UncontrolledFiles[0] : Files[0];
+		return SSourceControlCommon::GetSCCFileWidget(MoveTemp(FileState));
+	}
+
+	TArray<FSourceControlStateRef> Files;
+	TArray<FSourceControlStateRef> UncontrolledFiles;
+};
+
+
 typedef TSharedPtr<FUncontrolledChangelistTreeItem> FUncontrolledChangelistTreeItemPtr;
 typedef TSharedRef<FUncontrolledChangelistTreeItem> FUncontrolledChangelistTreeItemRef;
 typedef TSharedPtr<IChangelistTreeItem> FChangelistTreeItemPtr;
diff --git a/Engine/Source/Editor/SourceControlWindows/Private/SSourceControlFileDialog.cpp b/Engine/Source/Editor/SourceControlWindows/Private/SSourceControlFileDialog.cpp
--- a/Engine/Source/Editor/SourceControlWindows/Private/SSourceControlFileDialog.cpp
+++ b/Engine/Source/Editor/SourceControlWindows/Private/SSourceControlFileDialog.cpp
@@ -3,7 +3,7 @@
 #include "SSourceControlFileDialog.h"
 
 #include "Framework/Application/SlateApplication.h"
-#include "SFileTableRow.h"
+#include "SSourceControlChangelistRows.h"
 #include "Styling/AppStyle.h"
 #include "SWarningOrErrorBox.h"
 #include "Widgets/DeclarativeSyntaxSupport.h"
