Fix CIS non-unity build error in SSoucrControlChangelistRows and SSourceControlCommon

diff --git a/Engine/Source/Editor/SourceControlWindows/Private/SSourceControlChangelistRows.cpp b/Engine/Source/Editor/SourceControlWindows/Private/SSourceControlChangelistRows.cpp
--- a/Engine/Source/Editor/SourceControlWindows/Private/SSourceControlChangelistRows.cpp
+++ b/Engine/Source/Editor/SourceControlWindows/Private/SSourceControlChangelistRows.cpp
@@ -5,6 +5,9 @@
 #include "ISourceControlModule.h"
 #include "UncontrolledChangelistsModule.h"
 #include "SourceControlOperations.h"
+#include "Widgets/Images/SImage.h"
+#include "Widgets/Text/STextBlock.h"
+#include "Algo/Transform.h"
 
 #define LOCTEXT_NAMESPACE "SourceControlChangelistRow"
 
diff --git a/Engine/Source/Editor/SourceControlWindows/Private/SSourceControlCommon.cpp b/Engine/Source/Editor/SourceControlWindows/Private/SSourceControlCommon.cpp
--- a/Engine/Source/Editor/SourceControlWindows/Private/SSourceControlCommon.cpp
+++ b/Engine/Source/Editor/SourceControlWindows/Private/SSourceControlCommon.cpp
@@ -18,8 +18,10 @@
 #include "Widgets/Images/SImage.h"
 #include "Widgets/Images/SLayeredImage.h"
 #include "Widgets/Layout/SBox.h"
-#include"Framework/Notifications/NotificationManager.h"
+#include "Framework/Docking/TabManager.h"
+#include "Framework/Notifications/NotificationManager.h"
 #include "Misc/ScopedSlowTask.h"
+#include "Editor.h"
 
 #define LOCTEXT_NAMESPACE "SourceControlChangelist"
 
