Fixed CIS non-unity build in SSourceControlChangerlisRows.h

diff --git a/Engine/Source/Editor/SourceControlWindows/Private/SSourceControlChangelistRows.h b/Engine/Source/Editor/SourceControlWindows/Private/SSourceControlChangelistRows.h
--- a/Engine/Source/Editor/SourceControlWindows/Private/SSourceControlChangelistRows.h
+++ b/Engine/Source/Editor/SourceControlWindows/Private/SSourceControlChangelistRows.h
@@ -6,6 +6,7 @@
 
 #include "Widgets/SCompoundWidget.h"
 #include "Widgets/Views/STableRow.h"
+#include "Widgets/Views/STreeView.h"
 #include "Input/DragAndDrop.h"
 #include "SSourceControlCommon.h"
 
