Adds support for multiple uncontrolled changelists
* Add creation of new uncontrolled changelists
* Add modification of uncontrolled changelists description
* Add uncontrolled changelists deletion

diff --git a/Engine/Source/Developer/UncontrolledChangelists/Private/UncontrolledChangelistsModule.cpp b/Engine/Source/Developer/UncontrolledChangelists/Private/UncontrolledChangelistsModule.cpp
--- a/Engine/Source/Developer/UncontrolledChangelists/Private/UncontrolledChangelistsModule.cpp
+++ b/Engine/Source/Developer/UncontrolledChangelists/Private/UncontrolledChangelistsModule.cpp
@@ -483,6 +483,78 @@ void FUncontrolledChangelistsModule::MoveFilesToControlledChangelist(const TArra
 	}
 }
 
+void FUncontrolledChangelistsModule::CreateUncontrolledChangelist(const FText& InDescription)
+{
+	if (!IsEnabled())
+	{
+		return;
+	}
+
+	// Default constructor will generate a new GUID.
+	FUncontrolledChangelist NewUncontrolledChangelist;
+	UncontrolledChangelistsStateCache.Emplace(MoveTemp(NewUncontrolledChangelist), MakeShared<FUncontrolledChangelistState>(NewUncontrolledChangelist, InDescription));
+
+	OnStateChanged();
+}
+
+void FUncontrolledChangelistsModule::EditUncontrolledChangelist(const FUncontrolledChangelist& InUncontrolledChangelist, const FText& InNewDescription)
+{
+	if (!IsEnabled())
+	{
+		return;
+	}
+
+	if (InUncontrolledChangelist.IsDefault())
+	{
+		UE_LOG(LogSourceControl, Error, TEXT("Cannot edit Default Uncontrolled Changelist."));
+		return;
+	}
+
+	FUncontrolledChangelistStateRef* UncontrolledChangelistState = UncontrolledChangelistsStateCache.Find(InUncontrolledChangelist);
+
+	if (UncontrolledChangelistState == nullptr)
+	{
+		UE_LOG(LogSourceControl, Error, TEXT("Cannot find Uncontrolled Changelist %s in cache."), *InUncontrolledChangelist.ToString());
+		return;
+	}
+
+	(*UncontrolledChangelistState)->SetDescription(InNewDescription);
+
+	OnStateChanged();
+}
+
+void FUncontrolledChangelistsModule::DeleteUncontrolledChangelist(const FUncontrolledChangelist& InUncontrolledChangelist)
+{
+	if (!IsEnabled())
+	{
+		return;
+	}
+
+	if (InUncontrolledChangelist.IsDefault())
+	{
+		UE_LOG(LogSourceControl, Error, TEXT("Cannot delete Default Uncontrolled Changelist."));
+		return;
+	}
+
+	FUncontrolledChangelistStateRef* UncontrolledChangelistState = UncontrolledChangelistsStateCache.Find(InUncontrolledChangelist);
+
+	if (UncontrolledChangelistState == nullptr)
+	{
+		UE_LOG(LogSourceControl, Error, TEXT("Cannot find Uncontrolled Changelist %s in cache."), *InUncontrolledChangelist.ToString());
+		return;
+	}
+
+	if ((*UncontrolledChangelistState)->ContainsFiles())
+	{
+		UE_LOG(LogSourceControl, Error, TEXT("Cannot delete Uncontrolled Changelist %s while it contains files."), *InUncontrolledChangelist.ToString());
+		return;
+	}
+
+	UncontrolledChangelistsStateCache.Remove(InUncontrolledChangelist);
+
+	OnStateChanged();
+}
+
 void FUncontrolledChangelistsModule::OnStateChanged()
 {
 	OnUncontrolledChangelistModuleChanged.Broadcast();
diff --git a/Engine/Source/Developer/UncontrolledChangelists/Public/UncontrolledChangelistsModule.h b/Engine/Source/Developer/UncontrolledChangelists/Public/UncontrolledChangelistsModule.h
--- a/Engine/Source/Developer/UncontrolledChangelists/Public/UncontrolledChangelistsModule.h
+++ b/Engine/Source/Developer/UncontrolledChangelists/Public/UncontrolledChangelistsModule.h
@@ -126,6 +126,25 @@ public:
 	 */
 	void MoveFilesToControlledChangelist(const TArray<FString>& InUncontrolledFiles, const FSourceControlChangelistPtr& InChangelist, TFunctionRef<bool(const TArray<FSourceControlStateRef>&)> InOpenConflictDialog);
 
+	/**
+	 * Creates a new Uncontrolled Changelist.
+	 * @param	InDescription	The description of the newly created Uncontrolled Changelist.
+	 */
+	void CreateUncontrolledChangelist(const FText& InDescription);
+	
+	/**
+     * Edits an Uncontrolled Changelist's description
+	 * @param	InUncontrolledChangelist	The Uncontrolled Changelist to modify. Should not be the default Uncontrolled Changelist.
+	 * @param	InNewDescription			The description to set.
+	 */
+	void EditUncontrolledChangelist(const FUncontrolledChangelist& InUncontrolledChangelist, const FText& InNewDescription);
+	
+	/**
+	 * Deletes an Uncontrolled Changelist.
+	 * @param	InUncontrolledChangelist	The Uncontrolled Changelist to delete. Should not be the default Uncontrolled Changelist and should not contain files.
+	 */
+	void DeleteUncontrolledChangelist(const FUncontrolledChangelist& InUncontrolledChangelist);
+
 private:
 	/**
 	 * Helper use by Startup task and OnAssetLoaded delegate.
diff --git a/Engine/Source/Editor/SourceControlWindows/Private/SSourceControlChangelists.cpp b/Engine/Source/Editor/SourceControlWindows/Private/SSourceControlChangelists.cpp
--- a/Engine/Source/Editor/SourceControlWindows/Private/SSourceControlChangelists.cpp
+++ b/Engine/Source/Editor/SourceControlWindows/Private/SSourceControlChangelists.cpp
@@ -319,7 +319,7 @@ void SSourceControlChangelistsWidget::Construct(const FArguments& InArgs)
 	UncontrolledChangelistExpandableArea = SNew(SExpandableChangelistArea)
 		.HeaderText_Lambda([this]() { return FText::Format(LOCTEXT("SourceControl_UncontrolledChangeLists", "Uncontrolled Changelists ({0})"), UncontrolledChangelistTreeNodes.Num()); })
 		.ChangelistView(UncontrolledChangelistTreeView.ToSharedRef())
-		.NewButtonVisibility(EVisibility::Collapsed) // Functionality is planned but not implemented yet.
+		.OnNewChangelist_Lambda([this]() { OnNewUncontrolledChangelist(); return FReply::Handled(); })
 		.OnNewChangelistTooltip(LOCTEXT("Create_New_Uncontrolled_Changelist", "Create a new uncontrolled changelist."))
 		.SearchButtonVisibility(EVisibility::Collapsed); // Functionality is planned but not fully implemented yet.
 
@@ -1767,6 +1767,83 @@ bool SSourceControlChangelistsWidget::CanValidateChangelist()
 	return Changelist != nullptr && Changelist->GetFilesStates().Num() > 0;
 }
 
+void SSourceControlChangelistsWidget::OnNewUncontrolledChangelist()
+{
+	FText UncontrolledChangelistDescription;
+	bool bOk = GetChangelistDescription(
+		nullptr,
+		LOCTEXT("SourceControl.UncontrolledChangelist.New.Title", "New Uncontrolled Changelist..."),
+		LOCTEXT("SourceControl.UncontrolledChangelist.New.Label", "Enter a description for the uncontrolled changelist:"),
+		UncontrolledChangelistDescription);
+
+	if (!bOk)
+	{
+		return;
+	}
+
+	ExecuteUncontrolledChangelistOperation(LOCTEXT("Creating_Uncontrolled_Changelist", "Creating uncontrolled changelist..."), [&]()
+	{
+		FUncontrolledChangelistsModule::Get().CreateUncontrolledChangelist(UncontrolledChangelistDescription);
+	});
+}
+
+void SSourceControlChangelistsWidget::OnEditUncontrolledChangelist()
+{
+	FUncontrolledChangelistStatePtr UncontrolledChangelistState = GetCurrentUncontrolledChangelistState();
+
+	if (UncontrolledChangelistState == nullptr)
+	{
+		return;
+	}
+
+	FText NewUncontrolledChangelistDescription = UncontrolledChangelistState->GetDisplayText();
+
+	bool bOk = GetChangelistDescription(
+		nullptr,
+		LOCTEXT("SourceControl.Uncontrolled.Changelist.New.Title2", "Edit Uncontrolled Changelist..."),
+		LOCTEXT("SourceControl.Uncontrolled.Changelist.New.Label2", "Enter a new description for the uncontrolled changelist:"),
+		NewUncontrolledChangelistDescription);
+
+	if (!bOk)
+	{
+		return;
+	}
+
+	ExecuteUncontrolledChangelistOperation(LOCTEXT("Updating_Uncontrolled_Changelist_Description", "Updating uncontrolled changelist description..."), [&NewUncontrolledChangelistDescription, &UncontrolledChangelistState]()
+	{
+		FUncontrolledChangelistsModule::Get().EditUncontrolledChangelist(UncontrolledChangelistState->Changelist, NewUncontrolledChangelistDescription);
+	});
+}
+
+bool SSourceControlChangelistsWidget::CanEditUncontrolledChangelist()
+{
+	FUncontrolledChangelistStatePtr UncontrolledChangelistState = GetCurrentUncontrolledChangelistState();
+
+	return (UncontrolledChangelistState != nullptr) && !UncontrolledChangelistState->Changelist.IsDefault();
+}
+
+void SSourceControlChangelistsWidget::OnDeleteUncontrolledChangelist()
+{
+	TOptional<FUncontrolledChangelist> UncontrolledChangelist = GetCurrentUncontrolledChangelist();
+
+	if (!UncontrolledChangelist.IsSet())
+	{
+		return;
+	}
+
+	ExecuteUncontrolledChangelistOperation(LOCTEXT("Deleting_Uncontrolled_Changelist", "Deleting uncontrolled changelist..."), [&UncontrolledChangelist]()
+	{
+		FUncontrolledChangelistsModule::Get().DeleteUncontrolledChangelist(UncontrolledChangelist.GetValue());
+	});
+}
+
+bool SSourceControlChangelistsWidget::CanDeleteUncontrolledChangelist()
+{
+	FUncontrolledChangelistStatePtr UncontrolledChangelistState = GetCurrentUncontrolledChangelistState();
+
+	return (UncontrolledChangelistState != nullptr) && !UncontrolledChangelistState->Changelist.IsDefault() && !UncontrolledChangelistState->ContainsFiles();
+}
+
 void SSourceControlChangelistsWidget::OnMoveFiles()
 {
 	TArray<FString> SelectedControlledFiles;
@@ -2103,6 +2180,27 @@ TSharedPtr<SWidget> SSourceControlChangelistsWidget::OnOpenContextMenu()
 				FCanExecuteAction::CreateSP(this, &SSourceControlChangelistsWidget::CanDeleteChangelist)));
 	}
 
+	if (bHasSelectedUncontrolledChangelist)
+	{
+		Section.AddMenuEntry(
+			"EditUncontrolledChangelist",
+			LOCTEXT("SourceControl_EditUncontrolledChangelist", "Edit Uncontrolled Changelist..."),
+			LOCTEXT("SourceControl_Edit_Uncontrolled_Changelist_Tooltip", "Edit an uncontrolled changelist description"),
+			FSlateIcon(),
+			FUIAction(
+				FExecuteAction::CreateSP(this, &SSourceControlChangelistsWidget::OnEditUncontrolledChangelist),
+				FCanExecuteAction::CreateSP(this, &SSourceControlChangelistsWidget::CanEditUncontrolledChangelist)));
+
+		Section.AddMenuEntry(
+			"DeleteUncontrolledChangelist",
+			LOCTEXT("SourceControl_DeleteUncontrolledChangelist", "Delete Empty Changelist"),
+			LOCTEXT("SourceControl_Delete_Uncontrolled_Changelist_Tooltip", "Deletes an empty uncontrolled changelist"),
+			FSlateIcon(),
+			FUIAction(
+				FExecuteAction::CreateSP(this, &SSourceControlChangelistsWidget::OnDeleteUncontrolledChangelist),
+				FCanExecuteAction::CreateSP(this, &SSourceControlChangelistsWidget::CanDeleteUncontrolledChangelist)));
+	}
+
 	// Files-only operations
 	if(bHasSelectedFiles)
 	{
diff --git a/Engine/Source/Editor/SourceControlWindows/Private/SSourceControlChangelists.h b/Engine/Source/Editor/SourceControlWindows/Private/SSourceControlChangelists.h
--- a/Engine/Source/Editor/SourceControlWindows/Private/SSourceControlChangelists.h
+++ b/Engine/Source/Editor/SourceControlWindows/Private/SSourceControlChangelists.h
@@ -110,6 +110,13 @@ private:
 	void OnValidateChangelist();
 	bool CanValidateChangelist();
 
+	/** Uncontrolled Changelist operations */
+	void OnNewUncontrolledChangelist();
+	void OnEditUncontrolledChangelist();
+	bool CanEditUncontrolledChangelist();
+	void OnDeleteUncontrolledChangelist();
+	bool CanDeleteUncontrolledChangelist();
+
 	/** Changelist & File operations */
 	void OnRevertUnchanged();
 	bool CanRevertUnchanged();
