Fixed the 'Submit Window' saving the description of the default changelist into a new changelist without moving the files.
  - In the 'Submit' window, renamed the 'Apply' button to 'Save' button.
  - If the user saves with the 'default' changelist, a new changelist is created with the user description and files are moved into that changelist.
  - If the user saves with any other changelist, ony the description is edited.

diff --git a/Engine/Plugins/Developer/PerforceSourceControl/Source/PerforceSourceControl/Private/PerforceSourceControlChangelistState.h b/Engine/Plugins/Developer/PerforceSourceControl/Source/PerforceSourceControl/Private/PerforceSourceControlChangelistState.h
--- a/Engine/Plugins/Developer/PerforceSourceControl/Source/PerforceSourceControl/Private/PerforceSourceControlChangelistState.h
+++ b/Engine/Plugins/Developer/PerforceSourceControl/Source/PerforceSourceControl/Private/PerforceSourceControlChangelistState.h
@@ -36,6 +36,13 @@ public:
 	 */
 	virtual FText GetDescriptionText() const override;
 
+	/**
+	 * Returns whether the change list description can be saved and persisted. Some changelists are
+	 * special and the description cannot be persisted with the changelist. For example the P4 default
+	 * changelist description cannot be persisted and the user must create a new changelist.
+	 */
+	virtual bool SupportsPersistentDescription() const override { return !Changelist.IsDefault(); }
+
 	/**
 	 * Get a tooltip to describe this state
 	 * @returns	the text to display for this states tooltip
@@ -65,4 +72,4 @@ public:
 
 	/** The timestamp of the last update */
 	FDateTime TimeStamp;
-};
\ No newline at end of file
+};
diff --git a/Engine/Source/Developer/SourceControl/Public/ISourceControlChangelistState.h b/Engine/Source/Developer/SourceControl/Public/ISourceControlChangelistState.h
--- a/Engine/Source/Developer/SourceControl/Public/ISourceControlChangelistState.h
+++ b/Engine/Source/Developer/SourceControl/Public/ISourceControlChangelistState.h
@@ -45,6 +45,13 @@ public:
 	 */
 	virtual FText GetDescriptionText() const = 0;
 
+	/**
+	 * Returns whether the change list description can be saved and persisted. Some changelists are
+	 * special and the description cannot be persisted with the changelist. For example the P4 default
+	 * changelist description cannot be persisted and the user must create a new changelist.
+	 */
+	virtual bool SupportsPersistentDescription() const { return true; }
+
 	/**
 	 * Get a tooltip to describe this state
 	 * @returns	the text to display for this states tooltip
diff --git a/Engine/Source/Editor/SourceControlWindows/Private/SSourceControlChangelists.cpp b/Engine/Source/Editor/SourceControlWindows/Private/SSourceControlChangelists.cpp
--- a/Engine/Source/Editor/SourceControlWindows/Private/SSourceControlChangelists.cpp
+++ b/Engine/Source/Editor/SourceControlWindows/Private/SSourceControlChangelists.cpp
@@ -1623,15 +1623,7 @@ void SSourceControlChangelistsWidget::OnSubmitChangelist()
 		.AllowUncheckFiles(false)
 		.AllowKeepCheckedOut(true)
 		.AllowSubmit(bValidationResult)
-		.OnSaveChangelistDescription(
-			FSourceControlSaveChangelistDescription::CreateLambda([this, &ChangelistState, &bValidationResult, &UserEditChangelistDescription](const FText& NewDescription)
-			{
-				// NOTE this is called from a modal dialog, so adding a slow task on top of it doesn't really look good. Just run a synchronous operation.
-				TSharedRef<FEditChangelist> EditChangelistOperation = ISourceControlOperation::Create<FEditChangelist>();
-				EditChangelistOperation->SetDescription(NewDescription);
-				ISourceControlModule::Get().GetProvider().Execute(EditChangelistOperation, ChangelistState->GetChangelist(), EConcurrency::Synchronous);
-				UserEditChangelistDescription = NewDescription;
-			}));
+		.AllowSaveAndClose(true);
 
 	NewWindow->SetContent(
 		SourceControlWidget
@@ -1639,13 +1631,15 @@ void SSourceControlChangelistsWidget::OnSubmitChangelist()
 
 	FSlateApplication::Get().AddModalWindow(NewWindow, NULL);
 
+	bool bSaveDescriptionOnSubmitFailure = false;
+	bool bCheckinSuccess = false;
+
 	if (SourceControlWidget->GetResult() == ESubmitResults::SUBMIT_ACCEPTED)
 	{
 		ISourceControlProvider& SourceControlProvider = ISourceControlModule::Get().GetProvider();
 		FChangeListDescription Description;
 		TSharedRef<FCheckIn> SubmitChangelistOperation = ISourceControlOperation::Create<FCheckIn>();
 		SubmitChangelistOperation->SetKeepCheckedOut(SourceControlWidget->WantToKeepCheckedOut());
-		bool bCheckinSuccess = false;
 
 		// Get the changelist description the user had when he hit the 'submit' button.
 		SourceControlWidget->FillChangeListDescription(Description);
@@ -1676,21 +1670,49 @@ void SSourceControlChangelistsWidget::OnSubmitChangelist()
 		}
 
 		// If something went wrong with the submit, try to preserve the changelist edited by the user (if he edited).
-		if (!bCheckinSuccess && !UserEditChangelistDescription.EqualTo(ChangelistDescriptionToSubmit))
+		bSaveDescriptionOnSubmitFailure = (!bCheckinSuccess && !UserEditChangelistDescription.EqualTo(ChangelistDescriptionToSubmit));
+	}
+
+	if (SourceControlWidget->GetResult() == ESubmitResults::SUBMIT_SAVED || bSaveDescriptionOnSubmitFailure)
+	{
+		FChangeListDescription Description;
+		SourceControlWidget->FillChangeListDescription(Description);
+
+		if (ChangelistState->SupportsPersistentDescription())
 		{
 			TSharedRef<FEditChangelist> EditChangelistOperation = ISourceControlOperation::Create<FEditChangelist>();
-			EditChangelistOperation->SetDescription(UserEditChangelistDescription);
-			ISourceControlModule::Get().GetProvider().Execute(EditChangelistOperation, ChangelistState->GetChangelist(), EConcurrency::Synchronous);
+			EditChangelistOperation->SetDescription(Description.Description);
+			EditChangelistDescription(Description.Description, ChangelistState);
 		}
-
-		if (bCheckinSuccess)
+		else // Move everything to a new CL. Ex: the default P4 CL description cannot be saved.
 		{
-			// Clear the description saved by the 'submit window'. Useful when the submit window is opened from the Editor menu rather than the changelist window.
-			// Opening the 'submit window' from the Editor menu is intended for source controls that do not support changelists (SVN/Git), but remains available to
-			// all source controls at the moment.
-			SourceControlWidget->ClearChangeListDescription();
+			TArray<FString> FilesToMove;
+			Algo::Transform(ChangelistState->GetFilesStates(), FilesToMove, [](const TSharedRef<ISourceControlState>& FileState) { return FileState->GetFilename(); });
+
+			TSharedRef<FNewChangelist> NewChangelistOperation = ISourceControlOperation::Create<FNewChangelist>();
+			NewChangelistOperation->SetDescription(Description.Description);
+			Execute(LOCTEXT("Saving_Into_New_Changelist", "Saving to a new changelist..."), NewChangelistOperation, FilesToMove, EConcurrency::Synchronous, FSourceControlOperationComplete::CreateLambda(
+				[](const TSharedRef<ISourceControlOperation>& Operation, ECommandResult::Type InResult)
+				{
+					if (InResult == ECommandResult::Succeeded)
+					{
+						SSourceControlCommon::DisplaySourceControlOperationNotification(LOCTEXT("Saving_New_Changelist_Succeeded", "New changelist saved"), SNotificationItem::CS_Success);
+					}
+					if (InResult == ECommandResult::Failed)
+					{
+						SSourceControlCommon::DisplaySourceControlOperationNotification(LOCTEXT("Saving_New_Changelist_Failed", "Failed to save to a new changelist."), SNotificationItem::CS_Fail);
+					}
+				}));
 		}
 	}
+	
+	if (bCheckinSuccess)
+	{
+		// Clear the description saved by the 'submit window'. Useful when the submit window is opened from the Editor menu rather than the changelist window.
+		// Opening the 'submit window' from the Editor menu is intended for source controls that do not support changelists (SVN/Git), but remains available to
+		// all source controls at the moment.
+		SourceControlWidget->ClearChangeListDescription();
+	}
 }
 
 bool SSourceControlChangelistsWidget::CanSubmitChangelist()
diff --git a/Engine/Source/Editor/SourceControlWindows/Private/SSourceControlSubmit.cpp b/Engine/Source/Editor/SourceControlWindows/Private/SSourceControlSubmit.cpp
--- a/Engine/Source/Editor/SourceControlWindows/Private/SSourceControlSubmit.cpp
+++ b/Engine/Source/Editor/SourceControlWindows/Private/SSourceControlSubmit.cpp
@@ -146,7 +146,7 @@ void SSourceControlSubmitWidget::Construct(const FArguments& InArgs)
 	const bool bAllowUncheckFiles = InArgs._AllowUncheckFiles.Get();
 	const bool bAllowKeepCheckedOut = InArgs._AllowKeepCheckedOut.Get();
 	const bool bShowChangelistValidation = !InArgs._ChangeValidationResult.Get().IsEmpty();
-	OnSaveChangelistDescription = InArgs._OnSaveChangelistDescription;
+	const bool bAllowSaveAndClose = InArgs._AllowSaveAndClose.Get();
 
 	for (const auto& Item : InArgs._Items.Get())
 	{
@@ -386,12 +386,12 @@ void SSourceControlSubmitWidget::Construct(const FArguments& InArgs)
 		+SUniformGridPanel::Slot(0,0)
 		[
 			SNew(SButton)
-			.Visibility(OnSaveChangelistDescription.IsBound() ? EVisibility::Visible : EVisibility::Collapsed)
+			.Visibility(bAllowSaveAndClose ? EVisibility::Visible : EVisibility::Collapsed)
 			.HAlign(HAlign_Center)
 			.ContentPadding(FAppStyle::GetMargin("StandardDialog.ContentPadding"))
-			.Text(NSLOCTEXT("SourceControl.SubmitPanel", "Apply", "Apply"))
-			.ToolTipText(NSLOCTEXT("SourceControl.SubmitPanel", "Apply_Tooltip", "Save the description without submitting."))
-			.OnClicked(FOnClicked::CreateLambda([this]() { OnSaveChangelistDescription.ExecuteIfBound(ChangeListDescriptionTextCtrl->GetText()); return FReply::Handled(); }))
+			.Text(NSLOCTEXT("SourceControl.SubmitPanel", "Save", "Save"))
+			.ToolTipText(NSLOCTEXT("SourceControl.SubmitPanel", "Save_Tooltip", "Save the description and close without submitting."))
+			.OnClicked(this, &SSourceControlSubmitWidget::SaveAndCloseClicked)
 		]
 		+SUniformGridPanel::Slot(1,0)
 		[
@@ -493,13 +493,13 @@ void SSourceControlSubmitWidget::OnDiffAgainstDepotSelected(TSharedPtr<FFileTree
 
 FReply SSourceControlSubmitWidget::OnKeyDown( const FGeometry& MyGeometry, const FKeyEvent& InKeyEvent )
 {
-   // Pressing escape returns as if the user clicked cancel
-   if ( InKeyEvent.GetKey() == EKeys::Escape )
-   {
-      return CancelClicked();
-   }
+	// Pressing escape returns as if the user clicked cancel
+	if (InKeyEvent.GetKey() == EKeys::Escape)
+	{
+		return CancelClicked();
+	}
 
-   return FReply::Unhandled();
+	return FReply::Unhandled();
 }
 
 TSharedRef<SWidget> SSourceControlSubmitWidget::GenerateWidgetForItemAndColumn(TSharedPtr<FFileTreeItem> Item, const FName ColumnID) const
@@ -565,7 +565,7 @@ ECheckBoxState SSourceControlSubmitWidget::GetToggleSelectedState() const
 	ECheckBoxState PendingState = ECheckBoxState::Checked;
 
 	// Iterate through the list of selected items
-	for (const auto& Item : ListViewItems)
+	for (const TSharedPtr<FFileTreeItem>& Item : ListViewItems)
 	{
 		if (Item->GetCheckBoxState() == ECheckBoxState::Unchecked)
 		{
@@ -582,7 +582,7 @@ ECheckBoxState SSourceControlSubmitWidget::GetToggleSelectedState() const
 
 void SSourceControlSubmitWidget::OnToggleSelectedCheckBox(ECheckBoxState InNewState)
 {
-	for (const auto& Item : ListViewItems)
+	for (const TSharedPtr<FFileTreeItem>& Item : ListViewItems)
 	{
 		Item->SetCheckBoxState(InNewState);
 	}
@@ -598,7 +598,7 @@ void SSourceControlSubmitWidget::FillChangeListDescription(FChangeListDescriptio
 	OutDesc.FilesForAdd.Empty();
 	OutDesc.FilesForSubmit.Empty();
 
-	for (const auto& Item : ListViewItems)
+	for (const TSharedPtr<FFileTreeItem>& Item : ListViewItems)
 	{
 		if (Item->GetCheckBoxState() == ECheckBoxState::Checked)
 		{
@@ -633,7 +633,6 @@ FReply SSourceControlSubmitWidget::SubmitClicked()
 	return FReply::Handled();
 }
 
-
 FReply SSourceControlSubmitWidget::CancelClicked()
 {
 	DialogResult = ESubmitResults::SUBMIT_CANCELED;
@@ -642,6 +641,13 @@ FReply SSourceControlSubmitWidget::CancelClicked()
 	return FReply::Handled();
 }
 
+FReply SSourceControlSubmitWidget::SaveAndCloseClicked()
+{
+	DialogResult = ESubmitResults::SUBMIT_SAVED;
+	ParentFrame.Pin()->RequestDestroyWindow();
+
+	return FReply::Handled();
+}
 
 bool SSourceControlSubmitWidget::IsSubmitEnabled() const
 {
diff --git a/Engine/Source/Editor/SourceControlWindows/Private/SSourceControlSubmit.h b/Engine/Source/Editor/SourceControlWindows/Private/SSourceControlSubmit.h
--- a/Engine/Source/Editor/SourceControlWindows/Private/SSourceControlSubmit.h
+++ b/Engine/Source/Editor/SourceControlWindows/Private/SSourceControlSubmit.h
@@ -25,7 +25,8 @@ namespace ESubmitResults
 	enum Type
 	{
 		SUBMIT_ACCEPTED,
-		SUBMIT_CANCELED
+		SUBMIT_CANCELED,
+		SUBMIT_SAVED,
 	};
 }
 
@@ -58,7 +59,7 @@ public:
 		, _AllowUncheckFiles(true)
 		, _AllowKeepCheckedOut(true)
 		, _AllowSubmit(true)
-		, _OnSaveChangelistDescription()
+		, _AllowSaveAndClose(false)
 	{}
 
 		SLATE_ATTRIBUTE(TSharedPtr<SWindow>, ParentWindow)
@@ -71,7 +72,7 @@ public:
 		SLATE_ATTRIBUTE(bool, AllowUncheckFiles)
 		SLATE_ATTRIBUTE(bool, AllowKeepCheckedOut)
 		SLATE_ATTRIBUTE(bool, AllowSubmit)
-		SLATE_EVENT(FSourceControlSaveChangelistDescription, OnSaveChangelistDescription)
+		SLATE_ATTRIBUTE(bool, AllowSaveAndClose)
 
 	SLATE_END_ARGS()
 
@@ -117,6 +118,9 @@ private:
 	/** Called when the settings of the dialog are to be ignored*/
 	FReply CancelClicked();
 
+	/** Called when the user click the 'Save' button. */
+	FReply SaveAndCloseClicked();
+
 	/** Called to check if the submit button is enabled or not. */
 	bool IsSubmitEnabled() const;
 
@@ -196,9 +200,6 @@ private:
 
 	/** Currently selected sorting mode */
 	EColumnSortMode::Type SortMode;
-
-	/** Delegate invoked to save the description of a changelist with the source control. */
-	FSourceControlSaveChangelistDescription OnSaveChangelistDescription;
 };
 
 class SSourceControlSubmitListRow : public SMultiColumnTableRow<TSharedPtr<FFileTreeItem>>
