Fixes moving uncontrolled files to existing changelist.

diff --git a/Engine/Source/Editor/SourceControlWindows/Private/SSourceControlChangelists.cpp b/Engine/Source/Editor/SourceControlWindows/Private/SSourceControlChangelists.cpp
--- a/Engine/Source/Editor/SourceControlWindows/Private/SSourceControlChangelists.cpp
+++ b/Engine/Source/Editor/SourceControlWindows/Private/SSourceControlChangelists.cpp
@@ -1979,16 +1979,10 @@ void SSourceControlChangelistsWidget::OnMoveFiles()
 			if (!SelectedControlledFiles.IsEmpty())
 			{
 				Execute(LOCTEXT("Moving_File_Between_Changelists", "Moving file(s) to the selected changelist..."), ISourceControlOperation::Create<FMoveToChangelist>(), Changelist, SelectedControlledFiles, EConcurrency::Synchronous, FSourceControlOperationComplete::CreateLambda(
-					[SelectedUncontrolledFiles, Changelist](const TSharedRef<ISourceControlOperation>& Operation, ECommandResult::Type InResult)
+					[](const TSharedRef<ISourceControlOperation>& Operation, ECommandResult::Type InResult)
 					{
 						if (InResult == ECommandResult::Succeeded)
 						{
-							// Perform an uncontrolled move only if the controlled file were move successfully.
-							if (!SelectedUncontrolledFiles.IsEmpty())
-							{
-								FUncontrolledChangelistsModule::Get().MoveFilesToControlledChangelist(SelectedUncontrolledFiles, Changelist, SSourceControlCommon::OpenConflictDialog);
-							}
-
 							SSourceControlCommon::DisplaySourceControlOperationNotification(LOCTEXT("Move_Files_Between_Changelist_Succeeded", "File(s) successfully moved to the selected changelist."), SNotificationItem::CS_Success);
 						}
 						else if (InResult == ECommandResult::Failed)
@@ -1997,6 +1991,13 @@ void SSourceControlChangelistsWidget::OnMoveFiles()
 						}
 					}));
 			}
+			else if (!SelectedUncontrolledFiles.IsEmpty())
+			{
+				ExecuteUncontrolledChangelistOperation(LOCTEXT("Moving_Uncontrolled_Files_To_Changelist", "Moving uncontrolled files..."), [&]()
+				{
+					FUncontrolledChangelistsModule::Get().MoveFilesToControlledChangelist(SelectedUncontrolledFiles, Changelist, SSourceControlCommon::OpenConflictDialog);
+				});
+			}
 		}
 		else if (MoveDestination->GetTreeItemType() == IChangelistTreeItem::UncontrolledChangelist)
 		{
