Implemented text search/filtering in the changelist window.

Changed the color of the uncontrolled changelist from gray to orange
Added an icon in the to the file status column header.

diff --git a/Engine/Source/Editor/EditorStyle/Private/StarshipStyle.cpp b/Engine/Source/Editor/EditorStyle/Private/StarshipStyle.cpp
--- a/Engine/Source/Editor/EditorStyle/Private/StarshipStyle.cpp
+++ b/Engine/Source/Editor/EditorStyle/Private/StarshipStyle.cpp
@@ -7068,8 +7068,8 @@ void FStarshipEditorStyle::FStyle::SetupSourceControlStyles()
 		Set("SourceControl.ChangelistsTab", new CORE_IMAGE_BRUSH_SVG("Starship/Common/check-circle", Icon16x16));
 		Set("SourceControl.Changelist", new CORE_IMAGE_BRUSH_SVG("Starship/SourceControl/SCC_Changelist", Icon16x16, FStyleColors::AccentRed));
 		Set("SourceControl.ShelvedChangelist", new CORE_IMAGE_BRUSH_SVG("Starship/SourceControl/SCC_Changelist", Icon16x16, FStyleColors::AccentBlue));
-		Set("SourceControl.UncontrolledChangelist", new CORE_IMAGE_BRUSH_SVG("Starship/SourceControl/SCC_Changelist", Icon32x32, FStyleColors::AccentGray));
-		Set("SourceControl.UncontrolledChangelist_Small", new CORE_IMAGE_BRUSH_SVG("Starship/SourceControl/SCC_Changelist", Icon16x16, FStyleColors::AccentGray));
+		Set("SourceControl.UncontrolledChangelist", new CORE_IMAGE_BRUSH_SVG("Starship/SourceControl/SCC_Changelist", Icon32x32, FStyleColors::AccentOrange));
+		Set("SourceControl.UncontrolledChangelist_Small", new CORE_IMAGE_BRUSH_SVG("Starship/SourceControl/SCC_Changelist", Icon16x16, FStyleColors::AccentOrange));
 		Set("SourceControl.OfflineFile_Small", new CORE_IMAGE_BRUSH_SVG("Starship/SourceControl/SourceControl", Icon16x16, FStyleColors::AccentRed));
 
 		Set("SourceControl.Add", new CORE_IMAGE_BRUSH_SVG( "Starship/SourceControl/SCC_ContentAdd",Icon16x16));
diff --git a/Engine/Source/Editor/SourceControlWindows/Private/SSourceControlChangelistRows.cpp b/Engine/Source/Editor/SourceControlWindows/Private/SSourceControlChangelistRows.cpp
--- a/Engine/Source/Editor/SourceControlWindows/Private/SSourceControlChangelistRows.cpp
+++ b/Engine/Source/Editor/SourceControlWindows/Private/SSourceControlChangelistRows.cpp
@@ -5,6 +5,7 @@
 #include "ISourceControlModule.h"
 #include "UncontrolledChangelistsModule.h"
 #include "SourceControlOperations.h"
+#include "Styling/AppStyle.h"
 #include "Widgets/Images/SImage.h"
 #include "Widgets/Text/STextBlock.h"
 #include "Algo/Transform.h"
@@ -47,6 +48,7 @@ void SChangelistTableRow::Construct(const FArguments& InArgs, const TSharedRef<S
 			[
 				SNew(STextBlock)
 				.Text(this, &SChangelistTableRow::GetChangelistText)
+				.HighlightText(InArgs._HighlightText)
 			]
 			+SHorizontalBox::Slot() // Files count.
 			.Padding(4, 0, 4, 0)
@@ -55,16 +57,23 @@ void SChangelistTableRow::Construct(const FArguments& InArgs, const TSharedRef<S
 				SNew(STextBlock)
 				.Text(FText::Format(INVTEXT("({0})"), TreeItem->GetFileCount()))
 			]
-			+SHorizontalBox::Slot()
+			+SHorizontalBox::Slot() // Description.
 			.Padding(2, 0, 0, 0)
 			.AutoWidth()
 			[
 				SNew(STextBlock)
 				.Text(this, &SChangelistTableRow::GetChangelistDescriptionText)
+				.HighlightText(InArgs._HighlightText)
 			]
 		], InOwner);
 }
 
+void SChangelistTableRow::PopulateSearchString(const FChangelistTreeItem& Item, TArray<FString>& OutStrings)
+{
+	OutStrings.Emplace(Item.GetDisplayText().ToString()); // The changelist number
+	OutStrings.Emplace(GetChangelistDescription(Item));   // The changelist description.
+}
+
 FText SChangelistTableRow::GetChangelistText() const
 {
 	return TreeItem->GetDisplayText();
@@ -72,12 +81,17 @@ FText SChangelistTableRow::GetChangelistText() const
 
 FText SChangelistTableRow::GetChangelistDescriptionText() const
 {
-	FString DescriptionString = TreeItem->GetDescriptionText().ToString();
+	return FText::FromString(GetChangelistDescription(*TreeItem));
+}
+
+FString SChangelistTableRow::GetChangelistDescription(const FChangelistTreeItem& Item)
+{
+	FString DescriptionString = Item.GetDescriptionText().ToString();
 	// Here we'll both remove \r\n (when edited from the dialog) and \n (when we get it from the SCC)
 	DescriptionString.ReplaceInline(TEXT("\r"), TEXT(""));
 	DescriptionString.ReplaceInline(TEXT("\n"), TEXT(" "));
 	DescriptionString.TrimEndInline();
-	return FText::FromString(DescriptionString);
+	return DescriptionString;
 }
 
 FReply SChangelistTableRow::OnDrop(const FGeometry& InGeometry, const FDragDropEvent& InDragDropEvent)
@@ -139,7 +153,7 @@ void SUncontrolledChangelistTableRow::Construct(const FArguments& InArgs, const
 			FAppStyle::GetBrush(TreeItem->UncontrolledChangelistState->GetSmallIconName()) :
 			FAppStyle::GetBrush("SourceControl.Changelist");
 
-	SetToolTipText(GetChangelistDescriptionText());
+	SetToolTipText(GetChangelistText());
 
 	STableRow<FChangelistTreeItemPtr>::Construct(
 		STableRow<FChangelistTreeItemPtr>::FArguments()
@@ -161,6 +175,7 @@ void SUncontrolledChangelistTableRow::Construct(const FArguments& InArgs, const
 			[
 				SNew(STextBlock)
 				.Text(this, &SUncontrolledChangelistTableRow::GetChangelistText)
+				.HighlightText(InArgs._HighlightText)
 			]
 			+SHorizontalBox::Slot() // Files/Offline file count.
 			.Padding(4, 0, 4, 0)
@@ -172,21 +187,17 @@ void SUncontrolledChangelistTableRow::Construct(const FArguments& InArgs, const
 		], InOwner);
 }
 
-FText SUncontrolledChangelistTableRow::GetChangelistText() const
+void SUncontrolledChangelistTableRow::PopulateSearchString(const FUncontrolledChangelistTreeItem& Item, TArray<FString>& OutStrings)
 {
-	return TreeItem->GetDisplayText();
+	OutStrings.Emplace(Item.GetDisplayText().ToString());
 }
 
-FText SUncontrolledChangelistTableRow::GetChangelistDescriptionText() const
+FText SUncontrolledChangelistTableRow::GetChangelistText() const
 {
-	FString DescriptionString = TreeItem->GetDescriptionText().ToString();
-	// Here we'll both remove \r\n (when edited from the dialog) and \n (when we get it from the SCC)
-	DescriptionString.ReplaceInline(TEXT("\r"), TEXT(""));
-	DescriptionString.ReplaceInline(TEXT("\n"), TEXT(" "));
-	DescriptionString.TrimEndInline();
-	return FText::FromString(DescriptionString);
+	return TreeItem->GetDisplayText();
 }
 
+
 FReply SUncontrolledChangelistTableRow::OnDrop(const FGeometry& InGeometry, const FDragDropEvent& InDragDropEvent)
 {
 	TSharedPtr<FSCCFileDragDropOp> Operation = InDragDropEvent.GetOperationAs<FSCCFileDragDropOp>();
@@ -210,6 +221,8 @@ void SFileTableRow::Construct(const FArguments& InArgs, const TSharedRef<STableV
 {
 	TreeItem = static_cast<FFileTreeItem*>(InArgs._TreeItemToVisualize.Get());
 
+	HighlightText = InArgs._HighlightText;
+
 	FSuperRowType::FArguments Args = FSuperRowType::FArguments()
 		.OnDragDetected(InArgs._OnDragDetected)
 		.ShowSelection(true);
@@ -230,19 +243,22 @@ TSharedRef<SWidget> SFileTableRow::GenerateWidgetForColumn(const FName& ColumnId
 	else if (ColumnId == SourceControlFileViewColumnId::Name)
 	{
 		return SNew(STextBlock)
-			.Text(this, &SFileTableRow::GetDisplayName);
+			.Text(this, &SFileTableRow::GetDisplayName)
+			.HighlightText(HighlightText);
 	}
 	else if (ColumnId == SourceControlFileViewColumnId::Path)
 	{
 		return SNew(STextBlock)
 			.Text(this, &SFileTableRow::GetDisplayPath)
-			.ToolTipText(this, &SFileTableRow::GetFilename);
+			.ToolTipText(this, &SFileTableRow::GetFilename)
+			.HighlightText(HighlightText);
 	}
 	else if (ColumnId == SourceControlFileViewColumnId::Type)
 	{
 		return SNew(STextBlock)
 			.Text(this, &SFileTableRow::GetDisplayType)
-			.ColorAndOpacity(this, &SFileTableRow::GetDisplayColor);
+			.ColorAndOpacity(this, &SFileTableRow::GetDisplayColor)
+			.HighlightText(HighlightText);
 	}
 	else
 	{
@@ -250,6 +266,13 @@ TSharedRef<SWidget> SFileTableRow::GenerateWidgetForColumn(const FName& ColumnId
 	}
 }
 
+void SFileTableRow::PopulateSearchString(const FFileTreeItem& Item, TArray<FString>& OutStrings)
+{
+	OutStrings.Emplace(Item.GetAssetName().ToString()); // Name
+	OutStrings.Emplace(Item.GetAssetPath().ToString()); // Path
+	OutStrings.Emplace(Item.GetAssetType().ToString()); // Type
+}
+
 FText SFileTableRow::GetDisplayName() const
 {
 	return TreeItem->GetAssetName();
@@ -290,6 +313,8 @@ void SFileTableRow::OnDragLeave(FDragDropEvent const& InDragDropEvent)
 
 void SShelvedFilesTableRow::Construct(const FArguments& InArgs, const TSharedRef<STableViewBase>& InOwnerTableView)
 {
+	TreeItem = static_cast<FShelvedChangelistTreeItem*>(InArgs._TreeItemToVisualize.Get());
+
 	STableRow<FChangelistTreeItemPtr>::Construct(
 		STableRow<FChangelistTreeItemPtr>::FArguments()
 		.Content()
@@ -301,23 +326,30 @@ void SShelvedFilesTableRow::Construct(const FArguments& InArgs, const TSharedRef
 				.Padding(5, 0, 0, 0)
 				[
 					SNew(SImage)
-					.Image(InArgs._Icon)
+					.Image(FAppStyle::GetBrush("SourceControl.ShelvedChangelist"))
 				]
 				+SHorizontalBox::Slot()
 				.Padding(2.0f, 1.0f)
 				.VAlign(VAlign_Center)
 				[
 					SNew(STextBlock)
-					.Text(InArgs._Text)
+					.Text(TreeItem->GetDisplayText())
+					.HighlightText(InArgs._HighlightText)
 				]
 		],
 		InOwnerTableView);
 }
 
+void SShelvedFilesTableRow::PopulateSearchString(const FShelvedChangelistTreeItem& Item, TArray<FString>& OutStrings)
+{
+	OutStrings.Emplace(Item.GetDisplayText().ToString());
+}
+
 
 void SOfflineFileTableRow::Construct(const FArguments& InArgs, const TSharedRef<STableViewBase>& InOwner)
 {
 	TreeItem = static_cast<FOfflineFileTreeItem*>(InArgs._TreeItemToVisualize.Get());
+	HighlightText = InArgs._HighlightText;
 
 	FSuperRowType::FArguments Args = FSuperRowType::FArguments().ShowSelection(true);
 	FSuperRowType::Construct(Args, InOwner);
@@ -338,19 +370,22 @@ TSharedRef<SWidget> SOfflineFileTableRow::GenerateWidgetForColumn(const FName& C
 	else if (ColumnId == SourceControlFileViewColumnId::Name)
 	{
 		return SNew(STextBlock)
-			.Text(this, &SOfflineFileTableRow::GetDisplayName);
+			.Text(this, &SOfflineFileTableRow::GetDisplayName)
+			.HighlightText(HighlightText);
 	}
 	else if (ColumnId == SourceControlFileViewColumnId::Path)
 	{
 		return SNew(STextBlock)
 			.Text(this, &SOfflineFileTableRow::GetDisplayPath)
-			.ToolTipText(this, &SOfflineFileTableRow::GetFilename);
+			.ToolTipText(this, &SOfflineFileTableRow::GetFilename)
+			.HighlightText(HighlightText);
 	}
 	else if (ColumnId == SourceControlFileViewColumnId::Type)
 	{
 		return SNew(STextBlock)
 			.Text(this, &SOfflineFileTableRow::GetDisplayType)
-			.ColorAndOpacity(this, &SOfflineFileTableRow::GetDisplayColor);
+			.ColorAndOpacity(this, &SOfflineFileTableRow::GetDisplayColor)
+			.HighlightText(HighlightText);
 	}
 	else
 	{
@@ -358,6 +393,13 @@ TSharedRef<SWidget> SOfflineFileTableRow::GenerateWidgetForColumn(const FName& C
 	}
 }
 
+void SOfflineFileTableRow::PopulateSearchString(const FOfflineFileTreeItem& Item, TArray<FString>& OutStrings)
+{
+	OutStrings.Emplace(Item.GetDisplayName().ToString()); // Name
+	OutStrings.Emplace(Item.GetDisplayPath().ToString()); // Path
+	OutStrings.Emplace(Item.GetDisplayType().ToString()); // Type
+}
+
 FText SOfflineFileTableRow::GetDisplayName() const
 {
 	return TreeItem->GetDisplayName();
diff --git a/Engine/Source/Editor/SourceControlWindows/Private/SSourceControlChangelistRows.h b/Engine/Source/Editor/SourceControlWindows/Private/SSourceControlChangelistRows.h
--- a/Engine/Source/Editor/SourceControlWindows/Private/SSourceControlChangelistRows.h
+++ b/Engine/Source/Editor/SourceControlWindows/Private/SSourceControlChangelistRows.h
@@ -30,10 +30,12 @@ class SChangelistTableRow : public STableRow<TSharedPtr<IChangelistTreeItem>>
 public:
 	SLATE_BEGIN_ARGS(SChangelistTableRow)
 		: _TreeItemToVisualize()
+		, _HighlightText()
 		, _OnPostDrop()
 	{
 	}
 		SLATE_ARGUMENT(FChangelistTreeItemPtr, TreeItemToVisualize)
+		SLATE_ATTRIBUTE(FText, HighlightText)
 		SLATE_EVENT(FSimpleDelegate, OnPostDrop)
 	SLATE_END_ARGS()
 
@@ -48,9 +50,12 @@ public:
 	FText GetChangelistText() const;
 	FText GetChangelistDescriptionText() const;
 
+	static void PopulateSearchString(const FChangelistTreeItem& Item, TArray<FString>& OutStrings);
+
 protected:
 	//~ Begin STableRow Interface.
 	virtual FReply OnDrop(const FGeometry& InGeometry, const FDragDropEvent& InDragDropEvent) override;
+	static FString GetChangelistDescription(const FChangelistTreeItem& Item);
 
 private:
 	/** The info about the widget that we are visualizing. */
@@ -67,10 +72,12 @@ class SUncontrolledChangelistTableRow : public STableRow<FChangelistTreeItemPtr>
 public:
 	SLATE_BEGIN_ARGS(SUncontrolledChangelistTableRow)
 		: _TreeItemToVisualize()
+		, _HighlightText()
 		, _OnPostDrop()
 	{
 	}
 		SLATE_ARGUMENT(FChangelistTreeItemPtr, TreeItemToVisualize)
+		SLATE_ATTRIBUTE(FText, HighlightText)
 		SLATE_EVENT(FSimpleDelegate, OnPostDrop)
 	SLATE_END_ARGS()
 
@@ -83,7 +90,8 @@ public:
 	void Construct(const FArguments& InArgs, const TSharedRef<STableViewBase>& InOwner);
 
 	FText GetChangelistText() const;
-	FText GetChangelistDescriptionText() const;
+
+	static void PopulateSearchString(const FUncontrolledChangelistTreeItem& Item, TArray<FString>& OutStrings);
 
 protected:
 	//~ Begin STableRow Interface.
@@ -103,10 +111,12 @@ class SFileTableRow : public SMultiColumnTableRow<FChangelistTreeItemPtr>
 {
 public:
 	SLATE_BEGIN_ARGS(SFileTableRow)
-		: _TreeItemToVisualize()
+		: _TreeItemToVisualize(nullptr)
+		, _HighlightText()
 	{
 	}
 		SLATE_ARGUMENT(FChangelistTreeItemPtr, TreeItemToVisualize)
+		SLATE_ATTRIBUTE(FText, HighlightText)
 		SLATE_EVENT(FOnDragDetected, OnDragDetected)
 	SLATE_END_ARGS()
 
@@ -127,6 +137,8 @@ public:
 	FText GetDisplayType() const;
 	FSlateColor GetDisplayColor() const;
 
+	static void PopulateSearchString(const FFileTreeItem& Item, TArray<FString>& OutStrings);
+
 protected:
 	//~ Begin STableRow Interface.
 	virtual void OnDragEnter(FGeometry const& InGeometry, FDragDropEvent const& InDragDropEvent) override;
@@ -136,6 +148,9 @@ protected:
 private:
 	/** The info about the widget that we are visualizing. */
 	FFileTreeItem* TreeItem;
+
+	/** The text to highlight*/
+	TAttribute<FText> HighlightText;
 };
 
 
@@ -144,14 +159,23 @@ class SShelvedFilesTableRow : public STableRow<TSharedPtr<IChangelistTreeItem>>
 {
 public:
 	SLATE_BEGIN_ARGS(SShelvedFilesTableRow)
-		: _Icon(nullptr)
+		: _TreeItemToVisualize()
+		, _HighlightText()
 	{
 	}
-		SLATE_ARGUMENT(const FSlateBrush*, Icon)
-		SLATE_ARGUMENT(FText, Text)
+		SLATE_ARGUMENT(FChangelistTreeItemPtr, TreeItemToVisualize)
+		SLATE_ATTRIBUTE(FText, HighlightText)
 	SLATE_END_ARGS()
 
+public:
 	void Construct(const FArguments& InArgs, const TSharedRef<STableViewBase>& InOwnerTableView);
+	static void PopulateSearchString(const FShelvedChangelistTreeItem& Item, TArray<FString>& OutStrings);
+
+private:
+	FShelvedChangelistTreeItem* TreeItem;
+
+	/** The text to highlight*/
+	TAttribute<FText> HighlightText;
 };
 
 
@@ -161,9 +185,11 @@ class SOfflineFileTableRow : public SMultiColumnTableRow<FChangelistTreeItemPtr>
 public:
 	SLATE_BEGIN_ARGS(SOfflineFileTableRow)
 		: _TreeItemToVisualize()
+		, _HighlightText()
 	{
 	}
 		SLATE_ARGUMENT(FChangelistTreeItemPtr, TreeItemToVisualize)
+		SLATE_ATTRIBUTE(FText, HighlightText)
 	SLATE_END_ARGS()
 
 public:
@@ -183,7 +209,12 @@ public:
 	FText GetDisplayType() const;
 	FSlateColor GetDisplayColor() const;
 
+	static void PopulateSearchString(const FOfflineFileTreeItem& Item, TArray<FString>& OutStrings);
+
 private:
 	/** The info about the widget that we are visualizing. */
 	FOfflineFileTreeItem* TreeItem;
+
+	/** The text to highlight*/
+	TAttribute<FText> HighlightText;
 };
diff --git a/Engine/Source/Editor/SourceControlWindows/Private/SSourceControlChangelists.cpp b/Engine/Source/Editor/SourceControlWindows/Private/SSourceControlChangelists.cpp
--- a/Engine/Source/Editor/SourceControlWindows/Private/SSourceControlChangelists.cpp
+++ b/Engine/Source/Editor/SourceControlWindows/Private/SSourceControlChangelists.cpp
@@ -9,11 +9,11 @@
 
 #include "Widgets/DeclarativeSyntaxSupport.h"
 #include "Widgets/Images/SImage.h"
-#include "Widgets/Layout/SScrollBorder.h"
 #include "Widgets/Input/SButton.h"
 #include "Widgets/Input/SSearchBox.h"
 #include "Widgets/Layout/SBorder.h"
 #include "Widgets/Layout/SBox.h"
+#include "Widgets/Layout/SSpacer.h"
 #include "Widgets/Layout/SSplitter.h"
 #include "Widgets/Layout/SExpandableArea.h"
 
@@ -31,6 +31,7 @@
 #include "AssetToolsModule.h"
 #include "Misc/MessageDialog.h"
 #include "Misc/ScopedSlowTask.h"
+#include "Algo/AnyOf.h"
 #include "HAL/PlatformTime.h"
 #include "ToolMenus.h"
 
@@ -131,7 +132,7 @@ public:
 			.IsChecked(this, &SExpandableSearchButton::GetToggleButtonState)
 			.OnCheckStateChanged(this, &SExpandableSearchButton::OnToggleButtonStateChanged)
 			.Style(FAppStyle::Get(), "ToggleButtonCheckbox")
-			.Padding(4.0f)
+			.Padding(2.0f)
 			.ToolTipText(NSLOCTEXT("ExpandableSearchArea", "ExpandCollapseSearchButton", "Expands or collapses the search text box"))
 			[
 				SNew(SImage)
@@ -187,6 +188,7 @@ public:
 		, _OnNewChangelistTooltip()
 		, _NewButtonVisibility(EVisibility::Visible)
 		, _SearchButtonVisibility(EVisibility::Visible)
+		, _OnSearchTextChanged()
 	{}
 		/** Search box style (used to match the glass icon) */
 		SLATE_STYLE_ARGUMENT(FSearchBoxStyle, Style)
@@ -204,18 +206,21 @@ public:
 		SLATE_ARGUMENT(EVisibility, NewButtonVisibility)
 		/** Make the 'search' button visible or not. */
 		SLATE_ARGUMENT(EVisibility, SearchButtonVisibility)
+		/** Invoked whenever the searched text changes. */
+		SLATE_EVENT(FOnTextChanged, OnSearchTextChanged)
 	SLATE_END_ARGS()
 
 	void Construct(const FArguments& InArgs)
 	{
-		SearchBox = SNew(SSearchBox);
+		SearchBox = SNew(SSearchBox)
+			.OnTextChanged(InArgs._OnSearchTextChanged);
 
 		ChildSlot
 		[
 			SAssignNew(ExpandableArea, SExpandableArea)
 			.BorderImage(FAppStyle::Get().GetBrush("Brushes.Header"))
 			.BodyBorderImage(FAppStyle::Get().GetBrush("Brushes.Recessed"))
-			.HeaderPadding(FMargin(4.0f, 3.0f))
+			.HeaderPadding(FMargin(4.0f, 2.0f))
 			.AllowAnimatedTransition(false)
 			.HeaderContent()
 			[
@@ -252,11 +257,8 @@ public:
 				.AutoWidth()
 				.Padding(4.0f, 0.0f)
 				[
-					SNew(SBox)
-					.Visibility(InArgs._SearchButtonVisibility)
-					[
-						SNew(SExpandableSearchButton, SearchBox.ToSharedRef())
-					]
+					SNew(SExpandableSearchButton, SearchBox.ToSharedRef())
+						.Visibility(InArgs._SearchButtonVisibility)
 				]
 			]
 			.BodyContent()
@@ -286,6 +288,8 @@ public:
 	}
 
 	bool IsExpanded() const { return ExpandableArea->IsExpanded(); }
+	FText GetSearchedText() const { return SearchBox->GetText(); }
+	TSharedPtr<SSearchBox> GetSearchBox() { return SearchBox; }
 
 private:
 	TSharedPtr<SExpandableArea> ExpandableArea;
@@ -314,14 +318,29 @@ void SSourceControlChangelistsWidget::Construct(const FArguments& InArgs)
 		.ChangelistView(ChangelistTreeView.ToSharedRef())
 		.OnNewChangelist_Lambda([this](){ OnNewChangelist(); return FReply::Handled(); })
 		.OnNewChangelistTooltip(LOCTEXT("Create_New_Changelist", "Create a new changelist."))
-		.SearchButtonVisibility(EVisibility::Collapsed); // Functionality is planned but not fully implemented yet.
+		.SearchButtonVisibility(EVisibility::Visible)
+		.OnSearchTextChanged(this, &SSourceControlChangelistsWidget::OnChangelistSearchTextChanged);
 
 	UncontrolledChangelistExpandableArea = SNew(SExpandableChangelistArea)
 		.HeaderText_Lambda([this]() { return FText::Format(LOCTEXT("SourceControl_UncontrolledChangeLists", "Uncontrolled Changelists ({0})"), UncontrolledChangelistTreeNodes.Num()); })
 		.ChangelistView(UncontrolledChangelistTreeView.ToSharedRef())
 		.OnNewChangelist_Lambda([this]() { OnNewUncontrolledChangelist(); return FReply::Handled(); })
 		.OnNewChangelistTooltip(LOCTEXT("Create_New_Uncontrolled_Changelist", "Create a new uncontrolled changelist."))
-		.SearchButtonVisibility(EVisibility::Collapsed); // Functionality is planned but not fully implemented yet.
+		.SearchButtonVisibility(EVisibility::Visible) // Functionality is planned but not fully implemented yet.
+		.OnSearchTextChanged(this, &SSourceControlChangelistsWidget::OnUncontrolledChangelistSearchTextChanged);
+
+	ChangelistTextFilter = MakeShared<TTextFilter<const IChangelistTreeItem&>>(TTextFilter<const IChangelistTreeItem&>::FItemToStringArray::CreateSP(this, &SSourceControlChangelistsWidget::PopulateItemSearchStrings));
+	ChangelistTextFilter->OnChanged().AddSP(this, &SSourceControlChangelistsWidget::OnRefresh);
+
+	UncontrolledChangelistTextFilter = MakeShared<TTextFilter<const IChangelistTreeItem&>>(TTextFilter<const IChangelistTreeItem&>::FItemToStringArray::CreateSP(this, &SSourceControlChangelistsWidget::PopulateItemSearchStrings));
+	UncontrolledChangelistTextFilter->OnChanged().AddSP(this, &SSourceControlChangelistsWidget::OnRefresh);
+
+	FileTextFilter = MakeShared<TTextFilter<const IChangelistTreeItem&>>(TTextFilter<const IChangelistTreeItem&>::FItemToStringArray::CreateSP(this, &SSourceControlChangelistsWidget::PopulateItemSearchStrings));
+	FileTextFilter->OnChanged().AddSP(this, &SSourceControlChangelistsWidget::OnRefresh);
+
+	// Min/Max prevents making the Changelist Area too small and consequently prevent the file search box to go over the 'refresh' button.
+	const float MinChangelistAreaRatio = 0.2f;
+	const float MaxFileAreaRation = 1.0f - MinChangelistAreaRatio;
 
 	ChildSlot
 	[
@@ -333,13 +352,37 @@ void SSourceControlChangelistsWidget::Construct(const FArguments& InArgs)
 			.BorderImage(FAppStyle::GetBrush("ToolPanel.GroupBorder"))
 			.Padding(4)
 			[
-				SNew(SHorizontalBox)
-				+SHorizontalBox::Slot()
-				.HAlign(HAlign_Left)
-				.VAlign(VAlign_Center)
-				.AutoWidth()
+				SNew(SOverlay) // To align the file search box with the Changelist/File splitter, to compute ratio on the same width.
+				+SOverlay::Slot()
 				[
-					MakeToolBar()
+					SNew(SHorizontalBox)
+					+SHorizontalBox::Slot()
+					.HAlign(HAlign_Left)
+					.VAlign(VAlign_Center)
+					.AutoWidth()
+					[
+						MakeToolBar()
+					]
+				]
+
+				+SOverlay::Slot() // That slots has the same width than the SSplitter and use the same left/right ratio to keep the search box over the file section.
+				[
+					SNew(SHorizontalBox)
+					+SHorizontalBox::Slot()
+					.FillWidth(TAttribute<float>::CreateLambda([this, MinChangelistAreaRatio]() { return FMath::Max(MinChangelistAreaRatio, ChangelistAreaSize / (ChangelistAreaSize + FileAreaSize)); }))
+					.VAlign(VAlign_Center)
+					[
+						SNew(SSpacer) // This spacer uses the same width as the Changelist area on the left. They are kept in sync.
+					]
+
+					+SHorizontalBox::Slot()
+					.VAlign(VAlign_Center)
+					.FillWidth(TAttribute<float>::CreateLambda([this, MaxFileAreaRation]() { return FMath::Min(MaxFileAreaRation, FileAreaSize / (ChangelistAreaSize + FileAreaSize)); }))
+					[
+						SAssignNew(FileSearchBox, SSearchBox)
+						.HintText(LOCTEXT("Search_Files", "Search the files..."))
+						.OnTextChanged(this, &SSourceControlChangelistsWidget::OnFileSearchTextChanged)
+					]
 				]
 			]
 		]
@@ -367,9 +410,11 @@ void SSourceControlChangelistsWidget::Construct(const FArguments& InArgs)
 				// Left slot: Changelists and uncontrolled changelists areas
 				+SSplitter::Slot()
 				.Resizable(true)
-				.Value(0.30)
+				.SizeRule(SSplitter::FractionOfParent)
+				.Value_Lambda([this, MinChangelistAreaRatio]() { return FMath::Max(MinChangelistAreaRatio, ChangelistAreaSize); })
+				.OnSlotResized_Lambda([this](float Size) { ChangelistAreaSize = Size; })
 				[
-					SNew(SOverlay) // Visible when both Controlled and Uncontrolled changelists are enabled (Need to add a splitter)
+					SNew(SOverlay) // Visible when both Controlled and Uncontrolled changelists are enabled (Need to split vertical space)
 					+SOverlay::Slot()
 					[
 						SNew(SSplitter)
@@ -413,11 +458,11 @@ void SSourceControlChangelistsWidget::Construct(const FArguments& InArgs)
 				// Right slot: Files associated to the selected the changelist/uncontrolled changelist.
 				+SSplitter::Slot()
 				.Resizable(true)
+				.SizeRule(SSplitter::FractionOfParent)
+				.Value_Lambda([this, MaxFileAreaRation]() { return FMath::Min(MaxFileAreaRation, FileAreaSize); })
+				.OnSlotResized_Lambda([this](float Size) { FileAreaSize = Size; })
 				[
-					SNew(SScrollBorder, FileTreeView.ToSharedRef())
-					[
-						FileTreeView.ToSharedRef()
-					]
+					FileTreeView.ToSharedRef()
 				]
 			]
 		]
@@ -632,28 +677,57 @@ void SSourceControlChangelistsWidget::OnRefresh()
 
 	for (const TSharedRef<ISourceControlChangelistState>& ChangelistState : ChangelistsStates)
 	{
-		// Add a changelist.
 		TSharedRef<IChangelistTreeItem> ChangelistNode = MakeShared<FChangelistTreeItem>(ChangelistState);
+		TSharedPtr<IChangelistTreeItem> ShelvedFilesNode;
+		bool bShelvedFilesNodeFilteredOut = true;
+		bool bSearchForcedChangelistNodeExpansion = false;
+		if (ChangelistState->GetShelvedFilesStates().Num() > 0)
+		{
+			ShelvedFilesNode = MakeShared<FShelvedChangelistTreeItem>();
+			bShelvedFilesNodeFilteredOut = !ChangelistTextFilter->PassesFilter(*ShelvedFilesNode);
+			bSearchForcedChangelistNodeExpansion = !bShelvedFilesNodeFilteredOut && !ChangelistTextFilter->GetRawFilterText().IsEmpty();
+		}
+		
+		// Check if the 'changelist' node must be kept.
+		if (!ChangelistTextFilter->PassesFilter(*ChangelistNode) && bShelvedFilesNodeFilteredOut)
+		{
+			continue; // Both the changelist and the shelved nodes are filtered out.
+		}
+
+		// Add the changelist.
 		ChangelistTreeNodes.Add(ChangelistNode);
+		if (bSearchForcedChangelistNodeExpansion)
+		{
+			ChangelistTreeView->SetItemExpansion(ChangelistNode, /*bShouldExpand*/true);
+		}
 
+		// Add the changelist files.
 		for (const TSharedRef<ISourceControlState>& FileState : ChangelistState->GetFilesStates())
 		{
-			ChangelistNode->AddChild(MakeShared<FFileTreeItem>(FileState, bBeautifyPaths));
-			SlowTask.EnterProgressFrame();
-			bBeautifyPaths &= !SlowTask.ShouldCancel();
+			TSharedRef<FFileTreeItem> FileItem = MakeShared<FFileTreeItem>(FileState, bBeautifyPaths);
+			if (FileTextFilter->PassesFilter(*FileItem))
+			{
+				ChangelistNode->AddChild(FileItem);
+				SlowTask.EnterProgressFrame();
+				bBeautifyPaths &= !SlowTask.ShouldCancel();
+			}
 		}
 
-		if (ChangelistState->GetShelvedFilesStates().Num() > 0)
+		if (ShelvedFilesNode)
 		{
-			// Add a shelved files node under the changelist.
-			TSharedRef<IChangelistTreeItem> ShelvedFilesNode = MakeShared<FShelvedChangelistTreeItem>();
-			ChangelistNode->AddChild(ShelvedFilesNode);
+			// Add a shelved files node under the changelist node.
+			ChangelistNode->AddChild(ShelvedFilesNode.ToSharedRef());
 
+			// Add the shelved files.
 			for (const TSharedRef<ISourceControlState>& ShelvedFileState : ChangelistState->GetShelvedFilesStates())
 			{
-				ShelvedFilesNode->AddChild(MakeShared<FShelvedFileTreeItem>(ShelvedFileState, bBeautifyPaths));
-				SlowTask.EnterProgressFrame();
-				bBeautifyPaths &= !SlowTask.ShouldCancel();
+				TSharedRef<FShelvedFileTreeItem> ShelvedFileItem = MakeShared<FShelvedFileTreeItem>(ShelvedFileState, bBeautifyPaths);
+				if (FileTextFilter->PassesFilter(*ShelvedFileItem))
+				{
+					ShelvedFilesNode->AddChild(ShelvedFileItem);
+					SlowTask.EnterProgressFrame();
+					bBeautifyPaths &= !SlowTask.ShouldCancel();
+				}
 			}
 		}
 
@@ -665,20 +739,32 @@ void SSourceControlChangelistsWidget::OnRefresh()
 	{
 		// Add an uncontrolled changelist.
 		TSharedRef<IChangelistTreeItem> UncontrolledChangelistNode = MakeShared<FUncontrolledChangelistTreeItem>(UncontrolledChangelistState);
+		if (!UncontrolledChangelistTextFilter->PassesFilter(*UncontrolledChangelistNode))
+		{
+			continue;
+		}
 		UncontrolledChangelistTreeNodes.Add(UncontrolledChangelistNode);
 
 		for (const TSharedRef<ISourceControlState>& FileState : UncontrolledChangelistState->GetFilesStates())
 		{
-			UncontrolledChangelistNode->AddChild(MakeShared<FFileTreeItem>(FileState, bBeautifyPaths));
-			SlowTask.EnterProgressFrame();
-			bBeautifyPaths &= !SlowTask.ShouldCancel();
+			TSharedRef<FFileTreeItem> UncontrolledFileItem = MakeShared<FFileTreeItem>(FileState, bBeautifyPaths);
+			if (FileTextFilter->PassesFilter(*UncontrolledFileItem))
+			{
+				UncontrolledChangelistNode->AddChild(UncontrolledFileItem);
+				SlowTask.EnterProgressFrame();
+				bBeautifyPaths &= !SlowTask.ShouldCancel();
+			}
 		}
 
 		for (const FString& Filename : UncontrolledChangelistState->GetOfflineFiles())
 		{
-			UncontrolledChangelistNode->AddChild(MakeShared<FOfflineFileTreeItem>(Filename));
-			SlowTask.EnterProgressFrame();
-			bBeautifyPaths &= !SlowTask.ShouldCancel();
+			TSharedRef<FOfflineFileTreeItem> OfflineFileItem = MakeShared<FOfflineFileTreeItem>(Filename);
+			if (FileTextFilter->PassesFilter(*OfflineFileItem))
+			{
+				UncontrolledChangelistNode->AddChild(OfflineFileItem);
+				SlowTask.EnterProgressFrame();
+				bBeautifyPaths &= !SlowTask.ShouldCancel();
+			}
 		}
 
 		SlowTask.EnterProgressFrame();
@@ -2296,9 +2382,28 @@ TSharedRef<STreeView<FChangelistTreeItemPtr>> SSourceControlChangelistsWidget::C
 			+SHeaderRow::Column(SourceControlFileViewColumnId::Icon)
 			.DefaultLabel(FText::GetEmpty())
 			.FillSized(18)
+			.HeaderContentPadding(FMargin(0))
 			.SortPriority(this, &SSourceControlChangelistsWidget::GetColumnSortPriority, SourceControlFileViewColumnId::Icon)
 			.SortMode(this, &SSourceControlChangelistsWidget::GetColumnSortMode, SourceControlFileViewColumnId::Icon)
 			.OnSort(this, &SSourceControlChangelistsWidget::OnColumnSortModeChanged)
+			[
+				SNew(SHorizontalBox)
+				+SHorizontalBox::Slot()
+				.Padding(1, 0)
+				[
+					SNew(SBox)
+					.WidthOverride(16)
+					.HeightOverride(16)
+					.HAlign(HAlign_Center)
+					.VAlign(VAlign_Center)
+					.Visibility_Lambda([this](){ return GetColumnSortMode(SourceControlFileViewColumnId::Icon) == EColumnSortMode::None ? EVisibility::Visible : EVisibility::Collapsed; })
+					[
+						SNew(SImage)
+						.Image(FAppStyle::Get().GetBrush("SourceControl.ChangelistsTab"))
+						.ColorAndOpacity( FSlateColor::UseSubduedForeground() )
+					]
+				]
+			]
 
 			+SHeaderRow::Column(SourceControlFileViewColumnId::Name)
 			.DefaultLabel(LOCTEXT("Name", "Name"))
@@ -2330,33 +2435,38 @@ TSharedRef<ITableRow> SSourceControlChangelistsWidget::OnGenerateRow(TSharedPtr<
 	case IChangelistTreeItem::Changelist:
 		return SNew(SChangelistTableRow, OwnerTable)
 			.TreeItemToVisualize(InTreeItem)
+			.HighlightText_Lambda([this]() { return ChangelistExpandableArea->GetSearchedText(); })
 			.OnPostDrop(this, &SSourceControlChangelistsWidget::OnRefresh);
 
 	case IChangelistTreeItem::UncontrolledChangelist:
 		return SNew(SUncontrolledChangelistTableRow, OwnerTable)
 			.TreeItemToVisualize(InTreeItem)
+			.HighlightText_Lambda([this]() { return UncontrolledChangelistExpandableArea->GetSearchedText(); })
 			.OnPostDrop(this, &SSourceControlChangelistsWidget::OnRefresh);
 
 	case IChangelistTreeItem::File:
 		return SNew(SFileTableRow, OwnerTable)
 			.TreeItemToVisualize(InTreeItem)
+			.HighlightText_Lambda([this]() { return FileSearchBox->GetText(); })
 			.OnDragDetected(this, &SSourceControlChangelistsWidget::OnFilesDragged);
 
 	case IChangelistTreeItem::OfflineFile:
 		return SNew(SOfflineFileTableRow, OwnerTable)
-			.TreeItemToVisualize(InTreeItem);
+			.TreeItemToVisualize(InTreeItem)
+			.HighlightText_Lambda([this]() { return FileSearchBox->GetText(); });
 
 	case IChangelistTreeItem::ShelvedChangelist:
 		return SNew(SShelvedFilesTableRow, OwnerTable)
-			.Icon(FAppStyle::GetBrush("SourceControl.ShelvedChangelist"))
-			.Text(static_cast<const FShelvedChangelistTreeItem*>(InTreeItem.Get())->GetDisplayText());
+			.TreeItemToVisualize(InTreeItem)
+			.HighlightText_Lambda([this]() { return ChangelistExpandableArea->GetSearchedText(); });
 
 	case IChangelistTreeItem::ShelvedFile:
 		return SNew(SFileTableRow, OwnerTable)
-			.TreeItemToVisualize(InTreeItem);
+			.TreeItemToVisualize(InTreeItem)
+			.HighlightText_Lambda([this]() { return FileSearchBox->GetText(); });
 
 	default:
-		check(false);
+		checkNoEntry();
 	};
 
 	return SNew(STableRow<TSharedPtr<FString>>, OwnerTable);
@@ -2669,6 +2779,54 @@ void SSourceControlChangelistsWidget::SortFileView()
 	});
 }
 
+void SSourceControlChangelistsWidget::OnChangelistSearchTextChanged(const FText& InFilterText)
+{
+	ChangelistTextFilter->SetRawFilterText(InFilterText);
+	ChangelistExpandableArea->GetSearchBox()->SetError(ChangelistTextFilter->GetFilterErrorText());
+}
+
+void SSourceControlChangelistsWidget::OnUncontrolledChangelistSearchTextChanged(const FText& InFilterText)
+{
+	UncontrolledChangelistTextFilter->SetRawFilterText(InFilterText);
+	UncontrolledChangelistExpandableArea->GetSearchBox()->SetError(UncontrolledChangelistTextFilter->GetFilterErrorText());
+}
+
+void SSourceControlChangelistsWidget::OnFileSearchTextChanged(const FText& InFilterText)
+{
+	FileTextFilter->SetRawFilterText(InFilterText);
+	FileSearchBox->SetError(FileTextFilter->GetFilterErrorText());
+}
+
+void SSourceControlChangelistsWidget::PopulateItemSearchStrings(const IChangelistTreeItem& Item, TArray<FString>& OutStrings)
+{
+	switch (Item.GetTreeItemType())
+	{
+	case IChangelistTreeItem::Changelist:
+		SChangelistTableRow::PopulateSearchString(static_cast<const FChangelistTreeItem&>(Item), OutStrings);
+		break;
+
+	case IChangelistTreeItem::ShelvedChangelist:
+		SShelvedFilesTableRow::PopulateSearchString(static_cast<const FShelvedChangelistTreeItem&>(Item), OutStrings);
+		break;
+
+	case IChangelistTreeItem::UncontrolledChangelist:
+		SUncontrolledChangelistTableRow::PopulateSearchString(static_cast<const FUncontrolledChangelistTreeItem&>(Item), OutStrings);
+		break;
+
+	case IChangelistTreeItem::File:
+	case IChangelistTreeItem::ShelvedFile:
+		SFileTableRow::PopulateSearchString(static_cast<const FFileTreeItem&>(Item), OutStrings);
+		break;
+
+	case IChangelistTreeItem::OfflineFile:
+		SOfflineFileTableRow::PopulateSearchString(static_cast<const FOfflineFileTreeItem&>(Item), OutStrings);
+		break;
+
+	default:
+		checkNoEntry();
+	}
+}
+
 void SSourceControlChangelistsWidget::SaveExpandedAndSelectionStates(FExpandedAndSelectionStates& OutStates)
 {
 	// Save the selected item from the 'changelists' tree.
diff --git a/Engine/Source/Editor/SourceControlWindows/Private/SSourceControlChangelists.h b/Engine/Source/Editor/SourceControlWindows/Private/SSourceControlChangelists.h
--- a/Engine/Source/Editor/SourceControlWindows/Private/SSourceControlChangelists.h
+++ b/Engine/Source/Editor/SourceControlWindows/Private/SSourceControlChangelists.h
@@ -9,9 +9,11 @@
 
 #include "ISourceControlProvider.h"
 #include "SSourceControlCommon.h"
+#include "Misc/TextFilter.h"
 
 class FChangelistGroupTreeItem;
 class SExpandableChangelistArea;
+class SSearchBox;
 
 class SChangelistTree : public STreeView<FChangelistTreeItemPtr>
 {
@@ -59,6 +61,11 @@ private:
 	void OnColumnSortModeChanged(const EColumnSortPriority::Type SortPriority, const FName& ColumnId, const EColumnSortMode::Type InSortMode);
 	void SortFileView();
 
+	void OnChangelistSearchTextChanged(const FText& InFilterText);
+	void OnUncontrolledChangelistSearchTextChanged(const FText& InFilterText);
+	void OnFileSearchTextChanged(const FText& InFilterText);
+	void PopulateItemSearchStrings(const IChangelistTreeItem& Item, TArray<FString>& OutStrings);
+
 	FReply OnFilesDragged(const FGeometry& InGeometry, const FPointerEvent& InMouseEvent);
 
 	void RequestRefresh();
@@ -197,6 +204,11 @@ private:
 	EColumnSortMode::Type PrimarySortMode = EColumnSortMode::Ascending;
 	EColumnSortMode::Type SecondarySortMode = EColumnSortMode::None;
 
+	TSharedPtr<TTextFilter<const IChangelistTreeItem&>> ChangelistTextFilter;
+	TSharedPtr<TTextFilter<const IChangelistTreeItem&>> UncontrolledChangelistTextFilter;
+	TSharedPtr<TTextFilter<const IChangelistTreeItem&>> FileTextFilter;
+	TSharedPtr<SSearchBox> FileSearchBox;
+
 	void StartRefreshStatus();
 	void TickRefreshStatus(double InDeltaTime);
 	void EndRefreshStatus();
@@ -204,4 +216,7 @@ private:
 	FText RefreshStatus;
 	bool bIsRefreshing = false;
 	double RefreshStatusStartSecs;
+
+	float ChangelistAreaSize = 0.3; // [0.0f, 1.0f]
+	float FileAreaSize = 0.7;
 };
diff --git a/Engine/Source/Editor/SourceControlWindows/Private/SSourceControlCommon.cpp b/Engine/Source/Editor/SourceControlWindows/Private/SSourceControlCommon.cpp
--- a/Engine/Source/Editor/SourceControlWindows/Private/SSourceControlCommon.cpp
+++ b/Engine/Source/Editor/SourceControlWindows/Private/SSourceControlCommon.cpp
@@ -161,6 +161,11 @@ void FFileTreeItem::RefreshAssetInformation()
 	PackageName = FText::FromString(TempPackageName);
 }
 
+FText FFileTreeItem::GetAssetName() const
+{
+	return AssetName;
+}
+
 FText FFileTreeItem::GetAssetName()
 {
 	const FTimespan CurrentTime = FTimespan::FromSeconds(FPlatformTime::Seconds());
diff --git a/Engine/Source/Editor/SourceControlWindows/Private/SSourceControlCommon.h b/Engine/Source/Editor/SourceControlWindows/Private/SSourceControlCommon.h
--- a/Engine/Source/Editor/SourceControlWindows/Private/SSourceControlCommon.h
+++ b/Engine/Source/Editor/SourceControlWindows/Private/SSourceControlCommon.h
@@ -144,9 +144,12 @@ struct FFileTreeItem : public IChangelistTreeItem
 	/** Updates informations based on AssetData */
 	void RefreshAssetInformation();
 
-	/** Returns the asset name of the item */
+	/** Returns the asset name of the item. This might update the asset names from the asset registry. */
 	FText GetAssetName();
 
+	/** Returns the asset name. This returns the currently cached asset name.*/
+	FText GetAssetName() const;
+
 	/** Returns the asset path of the item */
 	FText GetAssetPath() const { return AssetPath; }
 
