Fixed updating the description of the default changelist from the changelist window.
  - The P4 default changelist cannot be saved, so a new changelist is created and all the files in the default changelist are moved in the new CL.

diff --git a/Engine/Source/Editor/SourceControlWindows/Private/SSourceControlChangelists.cpp b/Engine/Source/Editor/SourceControlWindows/Private/SSourceControlChangelists.cpp
--- a/Engine/Source/Editor/SourceControlWindows/Private/SSourceControlChangelists.cpp
+++ b/Engine/Source/Editor/SourceControlWindows/Private/SSourceControlChangelists.cpp
@@ -513,20 +513,43 @@ TSharedRef<SWidget> SSourceControlChangelistsWidget::MakeToolBar()
 
 void SSourceControlChangelistsWidget::EditChangelistDescription(const FText& InNewChangelistDescription, const FSourceControlChangelistStatePtr& InChangelistState)
 {
-	TSharedRef<FEditChangelist> EditChangelistOperation = ISourceControlOperation::Create<FEditChangelist>();
-	EditChangelistOperation->SetDescription(InNewChangelistDescription);
-	Execute(LOCTEXT("Updating_Changelist_Description", "Updating changelist description..."), EditChangelistOperation, InChangelistState->GetChangelist(), EConcurrency::Asynchronous, FSourceControlOperationComplete::CreateLambda(
-		[](const TSharedRef<ISourceControlOperation>& Operation, ECommandResult::Type InResult)
-		{
-			if (InResult == ECommandResult::Succeeded)
+	if (InChangelistState->SupportsPersistentDescription())
+	{
+		TSharedRef<FEditChangelist> EditChangelistOperation = ISourceControlOperation::Create<FEditChangelist>();
+		EditChangelistOperation->SetDescription(InNewChangelistDescription);
+		Execute(LOCTEXT("Updating_Changelist_Description", "Updating changelist description..."), EditChangelistOperation, InChangelistState->GetChangelist(), EConcurrency::Asynchronous, FSourceControlOperationComplete::CreateLambda(
+			[](const TSharedRef<ISourceControlOperation>& Operation, ECommandResult::Type InResult)
 			{
-				SSourceControlCommon::DisplaySourceControlOperationNotification(LOCTEXT("Update_Changelist_Description_Succeeded", "Changelist description successfully updated."), SNotificationItem::CS_Success);
-			}
-			else if (InResult == ECommandResult::Failed)
+				if (InResult == ECommandResult::Succeeded)
+				{
+					SSourceControlCommon::DisplaySourceControlOperationNotification(LOCTEXT("Update_Changelist_Description_Succeeded", "Changelist description successfully updated."), SNotificationItem::CS_Success);
+				}
+				else if (InResult == ECommandResult::Failed)
+				{
+					SSourceControlCommon::DisplaySourceControlOperationNotification(LOCTEXT("Update_Changelist_Description_Failed", "Failed to update changelist description."), SNotificationItem::CS_Fail);
+				}
+			}));
+	}
+	else // Move everything to a new CL. Ex: the default P4 CL description cannot be saved.
+	{
+		TArray<FString> FilesToMove;
+		Algo::Transform(InChangelistState->GetFilesStates(), FilesToMove, [](const TSharedRef<ISourceControlState>& FileState) { return FileState->GetFilename(); });
+
+		TSharedRef<FNewChangelist> NewChangelistOperation = ISourceControlOperation::Create<FNewChangelist>();
+		NewChangelistOperation->SetDescription(InNewChangelistDescription);
+		Execute(LOCTEXT("Saving_Into_New_Changelist", "Saving to a new changelist..."), NewChangelistOperation, FilesToMove, EConcurrency::Synchronous, FSourceControlOperationComplete::CreateLambda(
+			[](const TSharedRef<ISourceControlOperation>& Operation, ECommandResult::Type InResult)
 			{
-				SSourceControlCommon::DisplaySourceControlOperationNotification(LOCTEXT("Update_Changelist_Description_Failed", "Failed to update changelist description."), SNotificationItem::CS_Fail);
-			}
-		}));
+				if (InResult == ECommandResult::Succeeded)
+				{
+					SSourceControlCommon::DisplaySourceControlOperationNotification(LOCTEXT("Saving_New_Changelist_Succeeded", "New changelist saved"), SNotificationItem::CS_Success);
+				}
+				if (InResult == ECommandResult::Failed)
+				{
+					SSourceControlCommon::DisplaySourceControlOperationNotification(LOCTEXT("Saving_New_Changelist_Failed", "Failed to save to a new changelist."), SNotificationItem::CS_Fail);
+				}
+			}));
+	}
 }
 
 void SSourceControlChangelistsWidget::Tick(const FGeometry& AllottedGeometry, const double InCurrentTime, const float InDeltaTime)
@@ -1763,33 +1786,7 @@ void SSourceControlChangelistsWidget::OnSubmitChangelist()
 	{
 		FChangeListDescription Description;
 		SourceControlWidget->FillChangeListDescription(Description);
-
-		if (ChangelistState->SupportsPersistentDescription())
-		{
-			TSharedRef<FEditChangelist> EditChangelistOperation = ISourceControlOperation::Create<FEditChangelist>();
-			EditChangelistOperation->SetDescription(Description.Description);
-			EditChangelistDescription(Description.Description, ChangelistState);
-		}
-		else // Move everything to a new CL. Ex: the default P4 CL description cannot be saved.
-		{
-			TArray<FString> FilesToMove;
-			Algo::Transform(ChangelistState->GetFilesStates(), FilesToMove, [](const TSharedRef<ISourceControlState>& FileState) { return FileState->GetFilename(); });
-
-			TSharedRef<FNewChangelist> NewChangelistOperation = ISourceControlOperation::Create<FNewChangelist>();
-			NewChangelistOperation->SetDescription(Description.Description);
-			Execute(LOCTEXT("Saving_Into_New_Changelist", "Saving to a new changelist..."), NewChangelistOperation, FilesToMove, EConcurrency::Synchronous, FSourceControlOperationComplete::CreateLambda(
-				[](const TSharedRef<ISourceControlOperation>& Operation, ECommandResult::Type InResult)
-				{
-					if (InResult == ECommandResult::Succeeded)
-					{
-						SSourceControlCommon::DisplaySourceControlOperationNotification(LOCTEXT("Saving_New_Changelist_Succeeded", "New changelist saved"), SNotificationItem::CS_Success);
-					}
-					if (InResult == ECommandResult::Failed)
-					{
-						SSourceControlCommon::DisplaySourceControlOperationNotification(LOCTEXT("Saving_New_Changelist_Failed", "Failed to save to a new changelist."), SNotificationItem::CS_Fail);
-					}
-				}));
-		}
+		EditChangelistDescription(Description.Description, ChangelistState);
 	}
 	
 	if (bCheckinSuccess)
