Added the 'Last Saved' and 'User' column to the changelist window.
  - Implemented sorting and searching for the new columns.
  - Added a FUpdateStatus request to the source control provider when the user selects a changelist.  This gets the users that checked out the files along with extra file status.

Added support to enable/disable changelist window columns for 'Type', 'Last Saved' and 'User' columns and persisted the settings.
Added tooltips over the changelist row that display the content of a column. Useful when the value is clipped because the column is too small.
Added tooltips over the changelist row header.
Prevented creation of useless temporary strings in PerforceSourceControlOperations.cpp

diff --git a/Engine/Plugins/Developer/PerforceSourceControl/Source/PerforceSourceControl/Private/PerforceSourceControlOperations.cpp b/Engine/Plugins/Developer/PerforceSourceControl/Source/PerforceSourceControl/Private/PerforceSourceControlOperations.cpp
--- a/Engine/Plugins/Developer/PerforceSourceControl/Source/PerforceSourceControl/Private/PerforceSourceControlOperations.cpp
+++ b/Engine/Plugins/Developer/PerforceSourceControl/Source/PerforceSourceControl/Private/PerforceSourceControlOperations.cpp
@@ -1141,9 +1141,8 @@ static void ParseUpdateStatusResults(const FP4RecordSet& InRecords, const TArray
 	// Build up a map of any other branch states
 	for (const FP4Record& ClientRecord : InRecords)
 	{
-		FString FileName = ClientRecord(TEXT("clientFile"));
-
-		if (FileName.Len())
+		const FString& ClientFile = ClientRecord(TEXT("clientFile"));
+		if (ClientFile.Len())
 		{
 			// Local workspace file, we're only interested in other branches here
 			continue;
@@ -1154,7 +1153,7 @@ static void ParseUpdateStatusResults(const FP4RecordSet& InRecords, const TArray
 		const FString& OtherOpen = ClientRecord(TEXT("otherOpen"));
 
 		FString Branch;
-
+		FString FileName(ClientFile);
 		if (DepotFileName.Split(ContentRoot, &Branch, &FileName))
 		{
 			// Sanitize
diff --git a/Engine/Source/Developer/SourceControl/Public/ISourceControlModule.h b/Engine/Source/Developer/SourceControl/Public/ISourceControlModule.h
--- a/Engine/Source/Developer/SourceControl/Public/ISourceControlModule.h
+++ b/Engine/Source/Developer/SourceControl/Public/ISourceControlModule.h
@@ -98,7 +98,7 @@ public:
 	virtual void QueueStatusUpdate(const TArray<UPackage*>& InPackages) = 0;
 
 	/**
-	 * Queues a file to have its s7ource control status updated in the background.
+	 * Queues a file to have its source control status updated in the background.
 	 * @param	InFilenames	The files to queue.
 	 */
 	virtual void QueueStatusUpdate(const TArray<FString>& InFilenames) = 0;
diff --git a/Engine/Source/Editor/SourceControlWindows/Private/SSourceControlChangelistRows.cpp b/Engine/Source/Editor/SourceControlWindows/Private/SSourceControlChangelistRows.cpp
--- a/Engine/Source/Editor/SourceControlWindows/Private/SSourceControlChangelistRows.cpp
+++ b/Engine/Source/Editor/SourceControlWindows/Private/SSourceControlChangelistRows.cpp
@@ -12,11 +12,29 @@
 
 #define LOCTEXT_NAMESPACE "SourceControlChangelistRow"
 
+FName SourceControlFileViewColumn::Icon::Id() { return TEXT("Icon"); }
+FText SourceControlFileViewColumn::Icon::GetDisplayText() {return LOCTEXT("Name_Icon", "Source Control Status"); }
+FText SourceControlFileViewColumn::Icon::GetToolTipText() { return LOCTEXT("Icon_Column_Tooltip", "Displays the asset/file status"); }
 
-const FName SourceControlFileViewColumnId::Icon = TEXT("Icon");
-const FName SourceControlFileViewColumnId::Name = TEXT("Name");
-const FName SourceControlFileViewColumnId::Path = TEXT("Path");
-const FName SourceControlFileViewColumnId::Type = TEXT("Type");
+FName SourceControlFileViewColumn::Name::Id() { return TEXT("Name"); }
+FText SourceControlFileViewColumn::Name::GetDisplayText() { return LOCTEXT("Name_Column", "Name"); }
+FText SourceControlFileViewColumn::Name::GetToolTipText() { return LOCTEXT("Name_Column_Tooltip", "Displays the asset/file name"); }
+
+FName SourceControlFileViewColumn::Path::Id() { return TEXT("Path"); }
+FText SourceControlFileViewColumn::Path::GetDisplayText() { return LOCTEXT("Path_Column", "Path"); }
+FText SourceControlFileViewColumn::Path::GetToolTipText() { return LOCTEXT("Path_Column_Tooltip", "Displays the asset/file path"); }
+
+FName SourceControlFileViewColumn::Type::Id() { return TEXT("Type"); }
+FText SourceControlFileViewColumn::Type::GetDisplayText() { return LOCTEXT("Type_Column", "Type"); }
+FText SourceControlFileViewColumn::Type::GetToolTipText() { return LOCTEXT("Type_Column_Tooltip", "Displays the asset type"); }
+
+FName SourceControlFileViewColumn::LastModifiedTimestamp::Id() { return TEXT("LastModifiedTimestamp"); }
+FText SourceControlFileViewColumn::LastModifiedTimestamp::GetDisplayText() {return LOCTEXT("LastModifiedTimestamp_Column", "Last Saved"); }
+FText SourceControlFileViewColumn::LastModifiedTimestamp::GetToolTipText() { return LOCTEXT("LastMofiedTimestamp_Column_Tooltip", "Displays the last time the file/asset was saved on user hard drive"); }
+
+FName SourceControlFileViewColumn::CheckedOutByUser::Id() { return TEXT("CheckedOutByUser"); }
+FText SourceControlFileViewColumn::CheckedOutByUser::GetDisplayText() { return LOCTEXT("CheckedOutByUser_Column", "User"); }
+FText SourceControlFileViewColumn::CheckedOutByUser::GetToolTipText() { return LOCTEXT("CheckedOutByUser_Column_Tooltip", "Displays the other user(s) that checked out the file/asset, if any"); }
 
 
 void SChangelistTableRow::Construct(const FArguments& InArgs, const TSharedRef<STableViewBase>& InOwner)
@@ -231,7 +249,7 @@ void SFileTableRow::Construct(const FArguments& InArgs, const TSharedRef<STableV
 
 TSharedRef<SWidget> SFileTableRow::GenerateWidgetForColumn(const FName& ColumnId)
 {
-	if (ColumnId == SourceControlFileViewColumnId::Icon)
+	if (ColumnId == SourceControlFileViewColumn::Icon::Id())
 	{
 		return SNew(SBox)
 			.WidthOverride(16) // Small Icons are usually 16x16
@@ -240,26 +258,40 @@ TSharedRef<SWidget> SFileTableRow::GenerateWidgetForColumn(const FName& ColumnId
 				SSourceControlCommon::GetSCCFileWidget(TreeItem->FileState, TreeItem->IsShelved())
 			];
 	}
-	else if (ColumnId == SourceControlFileViewColumnId::Name)
+	else if (ColumnId == SourceControlFileViewColumn::Name::Id())
 	{
 		return SNew(STextBlock)
 			.Text(this, &SFileTableRow::GetDisplayName)
+			.ToolTipText(this, &SFileTableRow::GetDisplayName)
 			.HighlightText(HighlightText);
 	}
-	else if (ColumnId == SourceControlFileViewColumnId::Path)
+	else if (ColumnId == SourceControlFileViewColumn::Path::Id())
 	{
 		return SNew(STextBlock)
 			.Text(this, &SFileTableRow::GetDisplayPath)
 			.ToolTipText(this, &SFileTableRow::GetFilename)
 			.HighlightText(HighlightText);
 	}
-	else if (ColumnId == SourceControlFileViewColumnId::Type)
+	else if (ColumnId == SourceControlFileViewColumn::Type::Id())
 	{
 		return SNew(STextBlock)
 			.Text(this, &SFileTableRow::GetDisplayType)
+			.ToolTipText(this, &SFileTableRow::GetDisplayType)
 			.ColorAndOpacity(this, &SFileTableRow::GetDisplayColor)
 			.HighlightText(HighlightText);
 	}
+	else if (ColumnId == SourceControlFileViewColumn::LastModifiedTimestamp::Id())
+	{
+		return SNew(STextBlock)
+			.ToolTipText(this, &SFileTableRow::GetLastModifiedTimestamp)
+			.Text(this, &SFileTableRow::GetLastModifiedTimestamp);
+	}
+	else if (ColumnId == SourceControlFileViewColumn::CheckedOutByUser::Id())
+	{
+		return SNew(STextBlock)
+			.ToolTipText(this, &SFileTableRow::GetCheckedOutByUser)
+			.Text(this, &SFileTableRow::GetCheckedOutByUser);
+	}
 	else
 	{
 		return SNullWidget::NullWidget;
@@ -271,6 +303,8 @@ void SFileTableRow::PopulateSearchString(const FFileTreeItem& Item, TArray<FStri
 	OutStrings.Emplace(Item.GetAssetName().ToString()); // Name
 	OutStrings.Emplace(Item.GetAssetPath().ToString()); // Path
 	OutStrings.Emplace(Item.GetAssetType().ToString()); // Type
+	OutStrings.Emplace(Item.GetLastModifiedTimestamp().ToString());
+	OutStrings.Emplace(Item.GetCheckedOutByUser().ToString());
 }
 
 FText SFileTableRow::GetDisplayName() const
@@ -298,6 +332,16 @@ FSlateColor SFileTableRow::GetDisplayColor() const
 	return TreeItem->GetAssetTypeColor();
 }
 
+FText SFileTableRow::GetLastModifiedTimestamp() const
+{
+	return TreeItem->GetLastModifiedTimestamp();
+}
+
+FText SFileTableRow::GetCheckedOutByUser() const
+{
+	return TreeItem->GetCheckedOutByUser();
+}
+
 void SFileTableRow::OnDragEnter(FGeometry const& InGeometry, FDragDropEvent const& InDragDropEvent)
 {
 	TSharedPtr<FDragDropOperation> DragOperation = InDragDropEvent.GetOperation();
@@ -357,7 +401,7 @@ void SOfflineFileTableRow::Construct(const FArguments& InArgs, const TSharedRef<
 
 TSharedRef<SWidget> SOfflineFileTableRow::GenerateWidgetForColumn(const FName& ColumnId)
 {
-	if (ColumnId == SourceControlFileViewColumnId::Icon)
+	if (ColumnId == SourceControlFileViewColumn::Icon::Id())
 	{
 		return SNew(SBox)
 			.WidthOverride(16) // Small Icons are usually 16x16
@@ -367,26 +411,39 @@ TSharedRef<SWidget> SOfflineFileTableRow::GenerateWidgetForColumn(const FName& C
 				.Image(FAppStyle::GetBrush(FName("SourceControl.OfflineFile_Small")))
 			];
 	}
-	else if (ColumnId == SourceControlFileViewColumnId::Name)
+	else if (ColumnId == SourceControlFileViewColumn::Name::Id())
 	{
 		return SNew(STextBlock)
+			.ToolTipText(this, &SOfflineFileTableRow::GetDisplayName)
 			.Text(this, &SOfflineFileTableRow::GetDisplayName)
 			.HighlightText(HighlightText);
 	}
-	else if (ColumnId == SourceControlFileViewColumnId::Path)
+	else if (ColumnId == SourceControlFileViewColumn::Path::Id())
 	{
 		return SNew(STextBlock)
 			.Text(this, &SOfflineFileTableRow::GetDisplayPath)
 			.ToolTipText(this, &SOfflineFileTableRow::GetFilename)
 			.HighlightText(HighlightText);
 	}
-	else if (ColumnId == SourceControlFileViewColumnId::Type)
+	else if (ColumnId == SourceControlFileViewColumn::Type::Id())
 	{
 		return SNew(STextBlock)
 			.Text(this, &SOfflineFileTableRow::GetDisplayType)
+			.ToolTipText(this, &SOfflineFileTableRow::GetDisplayType)
 			.ColorAndOpacity(this, &SOfflineFileTableRow::GetDisplayColor)
 			.HighlightText(HighlightText);
 	}
+	else if (ColumnId == SourceControlFileViewColumn::LastModifiedTimestamp::Id())
+	{
+		return SNew(STextBlock)
+			.ToolTipText(this, &SOfflineFileTableRow::GetLastModifiedTimestamp)
+			.Text(this, &SOfflineFileTableRow::GetLastModifiedTimestamp);
+	}
+	else if (ColumnId == SourceControlFileViewColumn::CheckedOutByUser::Id())
+	{
+		return SNew(STextBlock)
+			.Text(FText::GetEmpty());
+	}
 	else
 	{
 		return SNullWidget::NullWidget;
@@ -398,6 +455,7 @@ void SOfflineFileTableRow::PopulateSearchString(const FOfflineFileTreeItem& Item
 	OutStrings.Emplace(Item.GetDisplayName().ToString()); // Name
 	OutStrings.Emplace(Item.GetDisplayPath().ToString()); // Path
 	OutStrings.Emplace(Item.GetDisplayType().ToString()); // Type
+	OutStrings.Emplace(Item.GetLastModifiedTimestamp().ToString());
 }
 
 FText SOfflineFileTableRow::GetDisplayName() const
@@ -425,4 +483,10 @@ FSlateColor SOfflineFileTableRow::GetDisplayColor() const
 	return TreeItem->GetDisplayColor();
 }
 
+FText SOfflineFileTableRow::GetLastModifiedTimestamp() const
+{
+	return TreeItem->GetLastModifiedTimestamp();
+}
+
+
 #undef LOCTEXT_NAMESPACE
diff --git a/Engine/Source/Editor/SourceControlWindows/Private/SSourceControlChangelistRows.h b/Engine/Source/Editor/SourceControlWindows/Private/SSourceControlChangelistRows.h
--- a/Engine/Source/Editor/SourceControlWindows/Private/SSourceControlChangelistRows.h
+++ b/Engine/Source/Editor/SourceControlWindows/Private/SSourceControlChangelistRows.h
@@ -10,17 +10,56 @@
 #include "Input/DragAndDrop.h"
 #include "SSourceControlCommon.h"
 
-/** Lists the unique column IDs used in the list view displaying controlled/uncontrolled changelist files. */
-namespace SourceControlFileViewColumnId
+/** Lists the unique columns used in the list view displaying controlled/uncontrolled changelist files. */
+namespace SourceControlFileViewColumn
 {
-	/** The icon column id. */
-	extern const FName Icon;
-	/** The file/asset name column id. */
-	extern const FName Name;
-	/** The file/asset path column id. */
-	extern const FName Path;
-	/** The file/asset type column Id. */
-	extern const FName Type;
+	/** The icon column. */
+	namespace Icon
+	{
+		FName Id();
+		FText GetDisplayText();
+		FText GetToolTipText();
+	};
+
+	/** The file/asset name column. */
+	namespace Name
+	{
+		FName Id();
+		FText GetDisplayText();
+		FText GetToolTipText();
+	};
+	
+	/** The file/asset path column. */
+	namespace Path
+	{
+		FName Id();
+		FText GetDisplayText();
+		FText GetToolTipText();
+	};
+
+	/** The file/asset type column. */
+	namespace Type
+	{
+		FName Id();
+		FText GetDisplayText();
+		FText GetToolTipText();
+	};
+
+	/** The last time the file/asset was modified column. */
+	namespace LastModifiedTimestamp
+	{
+		FName Id();
+		FText GetDisplayText();
+		FText GetToolTipText();
+	}
+
+	/** The user that has a given file in checkout column. */
+	namespace CheckedOutByUser
+	{
+		FName Id();
+		FText GetDisplayText();
+		FText GetToolTipText();
+	}
 }
 
 
@@ -136,6 +175,8 @@ public:
 	FText GetDisplayPath() const;
 	FText GetDisplayType() const;
 	FSlateColor GetDisplayColor() const;
+	FText GetLastModifiedTimestamp() const;
+	FText GetCheckedOutByUser() const;
 
 	static void PopulateSearchString(const FFileTreeItem& Item, TArray<FString>& OutStrings);
 
@@ -208,6 +249,8 @@ public:
 	FText GetDisplayPath() const;
 	FText GetDisplayType() const;
 	FSlateColor GetDisplayColor() const;
+	FText GetLastModifiedTimestamp() const;
+	FText GetCheckedOutByUser() const;
 
 	static void PopulateSearchString(const FOfflineFileTreeItem& Item, TArray<FString>& OutStrings);
 
diff --git a/Engine/Source/Editor/SourceControlWindows/Private/SSourceControlChangelists.cpp b/Engine/Source/Editor/SourceControlWindows/Private/SSourceControlChangelists.cpp
--- a/Engine/Source/Editor/SourceControlWindows/Private/SSourceControlChangelists.cpp
+++ b/Engine/Source/Editor/SourceControlWindows/Private/SSourceControlChangelists.cpp
@@ -28,17 +28,20 @@
 #include "SourceControlHelpers.h"
 #include "SourceControlPreferences.h"
 #include "SourceControlMenuContext.h"
+#include "SourceControlSettings.h"
 #include "AssetToolsModule.h"
 #include "Misc/MessageDialog.h"
 #include "Misc/ScopedSlowTask.h"
 #include "Algo/AnyOf.h"
 #include "HAL/PlatformTime.h"
+#include "HAL/FileManager.h"
 #include "ToolMenus.h"
-
+#include "UObject/UObjectGlobals.h"
 #include "SSourceControlSubmit.h"
 #include "Framework/Application/SlateApplication.h"
 #include "Framework/Notifications/NotificationManager.h"
 #include "Framework/Docking/TabManager.h"
+#include "Framework/MultiBox/MultiBoxBuilder.h"
 
 #define LOCTEXT_NAMESPACE "SourceControlChangelist"
 
@@ -99,6 +102,47 @@ FText UpdateChangelistDescriptionToSubmitIfNeeded(const bool bInValidationResult
 	return InChangelistDescription;
 }
 
+void UpdateLastModifiedTimestamp(IChangelistTreeItem& Item)
+{
+	auto GetTimestampText = [](const FString& Pathname)
+	{
+		if (IFileManager::Get().FileExists(*Pathname))
+		{
+			return FText::AsDateTime(IFileManager::Get().GetTimeStamp(*Pathname), EDateTimeStyle::Short);
+		}
+		return FText::GetEmpty();
+	};
+
+	switch (Item.GetTreeItemType())
+	{
+		// NOTE: We may not want to display the last saved timestamp for shelved files as this could lead users to think the file
+		//       was shelved at that time. Consider remove this case in the future if this is too confusing.
+		case IChangelistTreeItem::File:
+		case IChangelistTreeItem::ShelvedFile:
+		{
+			FFileTreeItem& FileItem = static_cast<FFileTreeItem&>(Item);
+			FileItem.SetLastModifiedTimestamp(GetTimestampText(FileItem.GetFileName().ToString()));
+			return;
+		}
+		break;
+
+		case IChangelistTreeItem::OfflineFile:
+		{
+			FOfflineFileTreeItem& OfflineItem = static_cast<FOfflineFileTreeItem&>(Item);
+			OfflineItem.SetLastModifiedTimestamp(GetTimestampText(OfflineItem.GetFilename()));
+			return;
+		}
+
+		case IChangelistTreeItem::Changelist:
+		case IChangelistTreeItem::UncontrolledChangelist:
+		case IChangelistTreeItem::ShelvedChangelist:
+			return; // We don't update timestamp for those nodes.
+
+		default:
+			checkNoEntry();
+	}
+}
+
 } // Anonymous namespace
 
 
@@ -307,7 +351,7 @@ void SSourceControlChangelistsWidget::Construct(const FArguments& InArgs)
 	SourceControlStateChangedDelegateHandle = SCCModule.GetProvider().RegisterSourceControlStateChanged_Handle(FSourceControlStateChanged::FDelegate::CreateSP(this, &SSourceControlChangelistsWidget::OnSourceControlStateChanged));
 	UncontrolledChangelistModule.OnUncontrolledChangelistModuleChanged.AddSP(this, &SSourceControlChangelistsWidget::OnSourceControlStateChanged);
 
-	PrimarySortedColumn = SourceControlFileViewColumnId::Name;
+	PrimarySortedColumn = SourceControlFileViewColumn::Name::Id();
 
 	ChangelistTreeView = CreateChangelistTreeView(ChangelistTreeNodes);
 	UncontrolledChangelistTreeView = CreateChangelistTreeView(UncontrolledChangelistTreeNodes);
@@ -502,7 +546,7 @@ TSharedRef<SWidget> SSourceControlChangelistsWidget::MakeToolBar()
 
 	ToolBarBuilder.AddToolBarButton(
 		FUIAction(
-			FExecuteAction::CreateLambda([this]() { RequestRefresh(); })),
+			FExecuteAction::CreateLambda([this]() { RequestChangelistsRefresh(); })),
 			NAME_None,
 			LOCTEXT("SourceControl_RefreshButton", "Refresh"),
 			LOCTEXT("SourceControl_RefreshButton_Tooltip", "Refreshes changelists from source control provider."),
@@ -567,7 +611,7 @@ void SSourceControlChangelistsWidget::Tick(const FGeometry& AllottedGeometry, co
 	{
 		if (ISourceControlModule::Get().IsEnabled() || FUncontrolledChangelistsModule::Get().IsEnabled())
 		{
-			RequestRefresh();
+			RequestChangelistsRefresh();
 		}
 		else
 		{
@@ -584,7 +628,7 @@ void SSourceControlChangelistsWidget::Tick(const FGeometry& AllottedGeometry, co
 	}
 }
 
-void SSourceControlChangelistsWidget::RequestRefresh()
+void SSourceControlChangelistsWidget::RequestChangelistsRefresh()
 {
 	bool bAnyProviderAvailable = false;
 
@@ -618,6 +662,62 @@ void SSourceControlChangelistsWidget::RequestRefresh()
 	}
 }
 
+void SSourceControlChangelistsWidget::RequestFileStatusRefresh(const IChangelistTreeItem& Changelist)
+{
+	TArray<FString> Pathnames;
+	if (Changelist.GetTreeItemType() == IChangelistTreeItem::Changelist)
+	{
+		const FChangelistTreeItem& ChangelistItem = static_cast<const FChangelistTreeItem&>(Changelist);
+		Algo::Transform(ChangelistItem.ChangelistState->GetFilesStates(), Pathnames, [](const TSharedRef<ISourceControlState>& FileState) { return FileState->GetFilename(); });
+	}
+	else if (Changelist.GetTreeItemType() == IChangelistTreeItem::ShelvedChangelist)
+	{
+		const FChangelistTreeItem& ChangelistItem = static_cast<const FChangelistTreeItem&>(Changelist);
+		Algo::Transform(ChangelistItem.ChangelistState->GetShelvedFilesStates(), Pathnames, [](const TSharedRef<ISourceControlState>& FileState) { return FileState->GetFilename(); });
+	}
+	else if (Changelist.GetTreeItemType() == IChangelistTreeItem::UncontrolledChangelist)
+	{
+		const FUncontrolledChangelistTreeItem& ChangelistItem = static_cast<const FUncontrolledChangelistTreeItem&>(Changelist);
+		Algo::Transform(ChangelistItem.UncontrolledChangelistState->GetFilesStates(), Pathnames, [](const TSharedRef<ISourceControlState>& FileState) { return FileState->GetFilename(); });
+		Algo::Transform(ChangelistItem.UncontrolledChangelistState->GetOfflineFiles(), Pathnames, [](const FString& Pathname) { return Pathname; });
+	}
+
+	if (!Pathnames.IsEmpty())
+	{
+		// Fire an async task to get the latest files status from source control to get the 'Checked Out By' status.
+		RequestFileStatusRefresh(Pathnames);
+	}
+}
+
+void SSourceControlChangelistsWidget::RequestFileStatusRefresh(const TArray<FString>& Pathnames)
+{
+	if (!ISourceControlModule::Get().IsEnabled())
+	{
+		return;
+	}
+
+	// If the changelist contains files.
+	if (!Pathnames.IsEmpty())
+	{
+		TSharedRef<FUpdateStatus> UpdateStatus = ISourceControlOperation::Create<FUpdateStatus>();
+
+		if (FUncontrolledChangelistsModule::Get().IsEnabled())
+		{
+			// Uncontrolled CL files are not opened in the P4 client view, force the status update of the files even if not opened to capture who has it in checked out.
+			UpdateStatus->SetForceUpdate(true);
+		}
+
+		OnStartSourceControlOperation(UpdateStatus, LOCTEXT("Updating_Changelist_Files_Status", "Updating changelist files status..."));
+
+		// Check if the source controlled states of those files changed outside the Editor. The Editor prevents a user from checking out assets that are out
+		// of date, but the user can checkout out from P4V directly. This request is meant to detect if a users checkout outdated assets outside the Editor. It will
+		// also detect non-exclusive files that are out of date (.ini files for example). That's an edge case and if this creates bottleneck, this can likely be disabled.
+		ISourceControlModule::Get().GetProvider().Execute(UpdateStatus, Pathnames, EConcurrency::Asynchronous, FSourceControlOperationComplete::CreateSP(this, &SSourceControlChangelistsWidget::OnEndSourceControlOperation));
+
+		// NOTE: The updated status are expected to come from OnSourceControlStateChanged() that is invoked when the source control internal state changes.
+	}
+}
+
 void SSourceControlChangelistsWidget::StartRefreshStatus()
 {
 	bIsRefreshing = true;
@@ -822,7 +922,7 @@ void SSourceControlChangelistsWidget::OnSourceControlProviderChanged(ISourceCont
 
 void SSourceControlChangelistsWidget::OnSourceControlStateChanged()
 {
-	// NOTE: No need to call RequestRefresh() to force the SCC to update internal states. We are being invoked because it was update, we just
+	// NOTE: No need to call RequestChangelistsRefresh() to force the SCC to update internal states. We are being invoked because it was update, we just
 	//       need to update the UI to reflect those state changes.
 	OnRefresh();
 }
@@ -831,7 +931,22 @@ void SSourceControlChangelistsWidget::OnChangelistsStatusUpdated(const TSharedRe
 {
 	// NOTE: This is invoked when the 'FUpdatePendingChangelistsStatus' completes. No need to refresh the tree views because OnSourceControlStateChanged() is also called.
 	OnEndSourceControlOperation(InOperation, InType);
-	EndRefreshStatus(); // TODO PL: Need to uniformize all operations status update. The 'Status Update' is different as it displays the time it takes.
+	EndRefreshStatus();
+
+	// Check if a changelist is selected.
+	TArray<TSharedPtr<IChangelistTreeItem>> SelectedChangelists = ChangelistTreeView->GetSelectedItems();
+	if (SelectedChangelists.Num() == 0)
+	{
+		// Check if an uncontrolled CL is selected. A changelist and an uncontrolled changelist cannot be selected at the same time.
+		SelectedChangelists = UncontrolledChangelistTreeView->GetSelectedItems();
+	}
+
+	// If there is a selected changelist (controlled or uncontrolled), refresh the file status with respect to the source control. The response will come
+	// as an update through OnSourceControlStateChanged() callback.
+	if (SelectedChangelists.Num() > 0)
+	{
+		RequestFileStatusRefresh(*SelectedChangelists[0]);
+	}
 }
 
 void SChangelistTree::Private_SetItemSelection(FChangelistTreeItemPtr TheItem, bool bShouldBeSelected, bool bWasUserDirected)
@@ -2365,7 +2480,21 @@ TSharedRef<SChangelistTree> SSourceControlChangelistsWidget::CreateChangelistTre
 
 TSharedRef<STreeView<FChangelistTreeItemPtr>> SSourceControlChangelistsWidget::CreateChangelistFilesView()
 {
-	return SNew(STreeView<FChangelistTreeItemPtr>)
+	USourceControlSettings* Settings = GetMutableDefault<USourceControlSettings>();
+	if (!Settings->bShowAssetTypeColumn)
+	{
+		FileViewHiddenColumnsList.Add(SourceControlFileViewColumn::Type::Id());
+	}
+	if (!Settings->bShowAssetLastModifiedTimeColumn)
+	{
+		FileViewHiddenColumnsList.Add(SourceControlFileViewColumn::LastModifiedTimestamp::Id());
+	}
+	if (!Settings->bShowAssetCheckedOutByColumn)
+	{
+		FileViewHiddenColumnsList.Add(SourceControlFileViewColumn::CheckedOutByUser::Id());
+	}
+
+	TSharedRef<STreeView<FChangelistTreeItemPtr>> FileView = SNew(STreeView<FChangelistTreeItemPtr>)
 		.ItemHeight(24.0f)
 		.TreeItemsSource(&FileTreeNodes)
 		.OnGenerateRow(this, &SSourceControlChangelistsWidget::OnGenerateRow)
@@ -2376,12 +2505,18 @@ TSharedRef<STreeView<FChangelistTreeItemPtr>> SSourceControlChangelistsWidget::C
 		.HeaderRow
 		(
 			SNew(SHeaderRow)
-			+SHeaderRow::Column(SourceControlFileViewColumnId::Icon)
-			.DefaultLabel(FText::GetEmpty())
+			.CanSelectGeneratedColumn(true)
+			.HiddenColumnsList(FileViewHiddenColumnsList)
+			.OnHiddenColumnsListChanged(this, &SSourceControlChangelistsWidget::OnFileViewHiddenColumnsListChanged)
+
+			+SHeaderRow::Column(SourceControlFileViewColumn::Icon::Id())
+			.DefaultLabel(SourceControlFileViewColumn::Icon::GetDisplayText()) // Displayed in the drop down menu to show/hide columns
+			.DefaultTooltip(SourceControlFileViewColumn::Icon::GetToolTipText())
+			.ShouldGenerateWidget(true) // Ensure the column cannot be hidden (grayed out in the show/hide drop down menu)
 			.FillSized(18)
 			.HeaderContentPadding(FMargin(0))
-			.SortPriority(this, &SSourceControlChangelistsWidget::GetColumnSortPriority, SourceControlFileViewColumnId::Icon)
-			.SortMode(this, &SSourceControlChangelistsWidget::GetColumnSortMode, SourceControlFileViewColumnId::Icon)
+			.SortPriority(this, &SSourceControlChangelistsWidget::GetColumnSortPriority, SourceControlFileViewColumn::Icon::Id())
+			.SortMode(this, &SSourceControlChangelistsWidget::GetColumnSortMode, SourceControlFileViewColumn::Icon::Id())
 			.OnSort(this, &SSourceControlChangelistsWidget::OnColumnSortModeChanged)
 			[
 				SNew(SHorizontalBox)
@@ -2393,36 +2528,86 @@ TSharedRef<STreeView<FChangelistTreeItemPtr>> SSourceControlChangelistsWidget::C
 					.HeightOverride(16)
 					.HAlign(HAlign_Center)
 					.VAlign(VAlign_Center)
-					.Visibility_Lambda([this](){ return GetColumnSortMode(SourceControlFileViewColumnId::Icon) == EColumnSortMode::None ? EVisibility::Visible : EVisibility::Collapsed; })
+					.Visibility_Lambda([this](){ return GetColumnSortMode(SourceControlFileViewColumn::Icon::Id()) == EColumnSortMode::None ? EVisibility::Visible : EVisibility::Collapsed; })
 					[
 						SNew(SImage)
 						.Image(FAppStyle::Get().GetBrush("SourceControl.ChangelistsTab"))
-						.ColorAndOpacity( FSlateColor::UseSubduedForeground() )
+						.ColorAndOpacity(FSlateColor::UseSubduedForeground())
 					]
 				]
 			]
 
-			+SHeaderRow::Column(SourceControlFileViewColumnId::Name)
-			.DefaultLabel(LOCTEXT("Name", "Name"))
-			.FillWidth(0.2f)
-			.SortPriority(this, &SSourceControlChangelistsWidget::GetColumnSortPriority, SourceControlFileViewColumnId::Name)
-			.SortMode(this, &SSourceControlChangelistsWidget::GetColumnSortMode, SourceControlFileViewColumnId::Name)
+			+SHeaderRow::Column(SourceControlFileViewColumn::Name::Id())
+			.DefaultLabel(SourceControlFileViewColumn::Name::GetDisplayText())
+			.DefaultTooltip(SourceControlFileViewColumn::Name::GetToolTipText())
+			.ShouldGenerateWidget(true) // Ensure the column cannot be hidden (grayed out in the show/hide drop down menu)
+			.FillWidth(1.5f)
+			.SortPriority(this, &SSourceControlChangelistsWidget::GetColumnSortPriority, SourceControlFileViewColumn::Name::Id())
+			.SortMode(this, &SSourceControlChangelistsWidget::GetColumnSortMode, SourceControlFileViewColumn::Name::Id())
+			.OnSort(this, &SSourceControlChangelistsWidget::OnColumnSortModeChanged)
+
+			+SHeaderRow::Column(SourceControlFileViewColumn::Path::Id())
+			.DefaultLabel(SourceControlFileViewColumn::Path::GetDisplayText())
+			.DefaultTooltip(SourceControlFileViewColumn::Path::GetToolTipText())
+			.ShouldGenerateWidget(true) // Ensure the column cannot be hidden (grayed out in the show/hide drop down menu)
+			.FillWidth(3.5f)
+			.SortPriority(this, &SSourceControlChangelistsWidget::GetColumnSortPriority, SourceControlFileViewColumn::Path::Id())
+			.SortMode(this, &SSourceControlChangelistsWidget::GetColumnSortMode, SourceControlFileViewColumn::Path::Id())
 			.OnSort(this, &SSourceControlChangelistsWidget::OnColumnSortModeChanged)
 
-			+SHeaderRow::Column(SourceControlFileViewColumnId::Path)
-			.DefaultLabel(LOCTEXT("Path", "Path"))
-			.FillWidth(0.6f)
-			.SortPriority(this, &SSourceControlChangelistsWidget::GetColumnSortPriority, SourceControlFileViewColumnId::Path)
-			.SortMode(this, &SSourceControlChangelistsWidget::GetColumnSortMode, SourceControlFileViewColumnId::Path)
+			+SHeaderRow::Column(SourceControlFileViewColumn::Type::Id())
+			.DefaultLabel(SourceControlFileViewColumn::Type::GetDisplayText())
+			.DefaultTooltip(SourceControlFileViewColumn::Type::GetToolTipText())
+			.FillWidth(2.0f)
+			.SortPriority(this, &SSourceControlChangelistsWidget::GetColumnSortPriority, SourceControlFileViewColumn::Type::Id())
+			.SortMode(this, &SSourceControlChangelistsWidget::GetColumnSortMode, SourceControlFileViewColumn::Type::Id())
 			.OnSort(this, &SSourceControlChangelistsWidget::OnColumnSortModeChanged)
 
-			+SHeaderRow::Column(SourceControlFileViewColumnId::Type)
-			.DefaultLabel(LOCTEXT("Type", "Type"))
-			.FillWidth(0.2f)
-			.SortPriority(this, &SSourceControlChangelistsWidget::GetColumnSortPriority, SourceControlFileViewColumnId::Type)
-			.SortMode(this, &SSourceControlChangelistsWidget::GetColumnSortMode, SourceControlFileViewColumnId::Type)
+			+SHeaderRow::Column(SourceControlFileViewColumn::LastModifiedTimestamp::Id())
+			.DefaultLabel(SourceControlFileViewColumn::LastModifiedTimestamp::GetDisplayText())
+			.DefaultTooltip(SourceControlFileViewColumn::LastModifiedTimestamp::GetToolTipText())
+			.FillWidth(1.5f)
+			.SortPriority(this, &SSourceControlChangelistsWidget::GetColumnSortPriority, SourceControlFileViewColumn::LastModifiedTimestamp::Id())
+			.SortMode(this, &SSourceControlChangelistsWidget::GetColumnSortMode, SourceControlFileViewColumn::LastModifiedTimestamp::Id())
 			.OnSort(this, &SSourceControlChangelistsWidget::OnColumnSortModeChanged)
-		);
+
+			+SHeaderRow::Column(SourceControlFileViewColumn::CheckedOutByUser::Id())
+			.DefaultLabel(SourceControlFileViewColumn::CheckedOutByUser::GetDisplayText())
+			.DefaultTooltip(SourceControlFileViewColumn::CheckedOutByUser::GetToolTipText())
+			.FillWidth(1.0f)
+			.SortPriority(this, &SSourceControlChangelistsWidget::GetColumnSortPriority, SourceControlFileViewColumn::CheckedOutByUser::Id())
+			.SortMode(this, &SSourceControlChangelistsWidget::GetColumnSortMode, SourceControlFileViewColumn::CheckedOutByUser::Id())
+			.OnSort(this, &SSourceControlChangelistsWidget::OnColumnSortModeChanged));
+
+	return FileView;
+}
+
+void SSourceControlChangelistsWidget::OnFileViewHiddenColumnsListChanged()
+{
+	if (FileTreeView && FileTreeView->GetHeaderRow())
+	{
+		USourceControlSettings* Settings = GetMutableDefault<USourceControlSettings>();
+		Settings->bShowAssetTypeColumn = true;
+		Settings->bShowAssetCheckedOutByColumn = true;
+		Settings->bShowAssetLastModifiedTimeColumn = true;
+
+		for (const FName& ColumnId : FileTreeView->GetHeaderRow()->GetHiddenColumnIds())
+		{
+			if (ColumnId == SourceControlFileViewColumn::Type::Id())
+			{
+				Settings->bShowAssetTypeColumn = false;
+			}
+			else if (ColumnId == SourceControlFileViewColumn::LastModifiedTimestamp::Id())
+			{
+				Settings->bShowAssetLastModifiedTimeColumn = false;
+			}
+			else if (ColumnId == SourceControlFileViewColumn::CheckedOutByUser::Id())
+			{
+				Settings->bShowAssetCheckedOutByColumn = false;
+			}
+		}
+		Settings->SaveConfig();
+	}
 }
 
 TSharedRef<ITableRow> SSourceControlChangelistsWidget::OnGenerateRow(TSharedPtr<IChangelistTreeItem> InTreeItem, const TSharedRef<STableViewBase>& OwnerTable)
@@ -2570,6 +2755,7 @@ void SSourceControlChangelistsWidget::OnChangelistSelectionChanged(TSharedPtr<IC
 			if (Child->GetTreeItemType() == DesiredChildrenType)
 			{
 				FileTreeNodes.Add(Child);
+				UpdateLastModifiedTimestamp(*Child);
 			}
 		}
 	};
@@ -2594,6 +2780,16 @@ void SSourceControlChangelistsWidget::OnChangelistSelectionChanged(TSharedPtr<IC
 			default:
 				break;
 		}
+
+		// If the user clicked the changelist himself (not a side effect of a refresh coming from a source control state change).
+		if (SelectionType != ESelectInfo::Direct)
+		{
+			// Fire a request to get more info about the  file status (checkout by, outdated, etc). This is expensive on the server to call periodically.
+			RequestFileStatusRefresh(*SelectedChangelist);
+
+			// Sort the view as it just changed.
+			SortFileView();
+		}
 	}
 
 	FileTreeView->RequestTreeRefresh();
@@ -2733,24 +2929,58 @@ void SSourceControlChangelistsWidget::SortFileView()
 		}
 	};
 
+	auto GetModifiedTimestamp = [](IChangelistTreeItem* Item)  -> FString
+	{
+		switch (Item->GetTreeItemType())
+		{
+		case IChangelistTreeItem::File:
+		case IChangelistTreeItem::ShelvedFile:
+			return static_cast<FFileTreeItem*>(Item)->GetLastModifiedTimestamp().ToString();
+		case IChangelistTreeItem::OfflineFile:
+			return static_cast<FOfflineFileTreeItem*>(Item)->GetLastModifiedTimestamp().ToString();
+		default:
+			return FString();
+		}
+	};
+
+	auto GetCheckedOutByUser = [](IChangelistTreeItem* Item)  -> FString
+	{
+		switch (Item->GetTreeItemType())
+		{
+		case IChangelistTreeItem::File:
+		case IChangelistTreeItem::ShelvedFile:
+			return static_cast<FFileTreeItem*>(Item)->GetCheckedOutByUser().ToString();
+		default:
+			return FString();
+		}
+	};
+
 	auto Compare = [&](const TSharedPtr<IChangelistTreeItem>& Lhs, const TSharedPtr<IChangelistTreeItem>& Rhs, const FName& ColumnId, EColumnSortMode::Type SortMode) -> bool
 	{
-		if (ColumnId == SourceControlFileViewColumnId::Icon)
+		if (ColumnId == SourceControlFileViewColumn::Icon::Id())
 		{
 			return SortMode == EColumnSortMode::Ascending ? OperatorLessIcon(Lhs.Get(), Rhs.Get()) : OperatorLessIcon(Rhs.Get(), Lhs.Get());
 		}
-		else if (ColumnId == SourceControlFileViewColumnId::Name)
+		else if (ColumnId == SourceControlFileViewColumn::Name::Id())
 		{
 			return SortMode == EColumnSortMode::Ascending ? GetName(Lhs.Get()) < GetName(Rhs.Get()) : GetName(Lhs.Get()) > GetName(Rhs.Get());
 		}
-		else if (ColumnId == SourceControlFileViewColumnId::Path)
+		else if (ColumnId == SourceControlFileViewColumn::Path::Id())
 		{
 			return SortMode == EColumnSortMode::Ascending ? GetPath(Lhs.Get()) < GetPath(Rhs.Get()) : GetPath(Lhs.Get()) > GetPath(Rhs.Get());
 		}
-		else if (ColumnId == SourceControlFileViewColumnId::Type)
+		else if (ColumnId == SourceControlFileViewColumn::Type::Id())
 		{
 			return SortMode == EColumnSortMode::Ascending ? GetType(Lhs.Get()) < GetType(Rhs.Get()) : GetType(Lhs.Get()) > GetType(Rhs.Get());
 		}
+		else if (ColumnId == SourceControlFileViewColumn::LastModifiedTimestamp::Id())
+		{
+			return SortMode == EColumnSortMode::Ascending ? GetModifiedTimestamp(Lhs.Get()) < GetModifiedTimestamp(Rhs.Get()) : GetModifiedTimestamp(Lhs.Get()) > GetModifiedTimestamp(Rhs.Get());
+		}
+		else if (ColumnId == SourceControlFileViewColumn::CheckedOutByUser::Id())
+		{
+			return SortMode == EColumnSortMode::Ascending ? GetCheckedOutByUser(Lhs.Get()) < GetCheckedOutByUser(Rhs.Get()) : GetCheckedOutByUser(Lhs.Get()) > GetCheckedOutByUser(Rhs.Get());
+		}
 		else
 		{
 			checkNoEntry();
diff --git a/Engine/Source/Editor/SourceControlWindows/Private/SSourceControlChangelists.h b/Engine/Source/Editor/SourceControlWindows/Private/SSourceControlChangelists.h
--- a/Engine/Source/Editor/SourceControlWindows/Private/SSourceControlChangelists.h
+++ b/Engine/Source/Editor/SourceControlWindows/Private/SSourceControlChangelists.h
@@ -14,6 +14,7 @@
 class FChangelistGroupTreeItem;
 class SExpandableChangelistArea;
 class SSearchBox;
+class USourceControlSettings;
 
 class SChangelistTree : public STreeView<FChangelistTreeItemPtr>
 {
@@ -55,6 +56,7 @@ private:
 	TSharedRef<ITableRow> OnGenerateRow(FChangelistTreeItemPtr InTreeItem, const TSharedRef<STableViewBase>& OwnerTable);
 	void OnGetFileChildren(FChangelistTreeItemPtr InParent, TArray<FChangelistTreeItemPtr>& OutChildren);
 	void OnGetChangelistChildren(FChangelistTreeItemPtr InParent, TArray<FChangelistTreeItemPtr>& OutChildren);
+	void OnFileViewHiddenColumnsListChanged();
 
 	EColumnSortPriority::Type GetColumnSortPriority(const FName ColumnId) const;
 	EColumnSortMode::Type GetColumnSortMode(const FName ColumnId) const;
@@ -68,7 +70,9 @@ private:
 
 	FReply OnFilesDragged(const FGeometry& InGeometry, const FPointerEvent& InMouseEvent);
 
-	void RequestRefresh();
+	void RequestChangelistsRefresh();
+	void RequestFileStatusRefresh(const IChangelistTreeItem& Changelist);
+	void RequestFileStatusRefresh(const TArray<FString>& Pathnames);
 	void OnRefresh();
 	void ClearChangelistsTree();
 
@@ -189,6 +193,7 @@ private:
 
 	/** Display the list of files associated to the selected changelist, uncontrolled changelist or shelved node. */
 	TSharedPtr<STreeView<FChangelistTreeItemPtr>> FileTreeView;
+	TArray<FName> FileViewHiddenColumnsList;
 
 	/** Source control state changed delegate handle */
 	FDelegateHandle SourceControlStateChangedDelegateHandle;
diff --git a/Engine/Source/Editor/SourceControlWindows/Private/SSourceControlCommon.cpp b/Engine/Source/Editor/SourceControlWindows/Private/SSourceControlCommon.cpp
--- a/Engine/Source/Editor/SourceControlWindows/Private/SSourceControlCommon.cpp
+++ b/Engine/Source/Editor/SourceControlWindows/Private/SSourceControlCommon.cpp
@@ -78,6 +78,16 @@ FFileTreeItem::FFileTreeItem(FSourceControlStateRef InFileState, bool bBeautifyP
 	RefreshAssetInformation();
 }
 
+FText FFileTreeItem::GetCheckedOutByUser() const
+{
+	FString Users;
+	if (FileState->IsCheckedOutOther(&Users))
+	{
+		return FText::FromString(Users);
+	}
+	return FText::GetEmpty();
+}
+
 void FFileTreeItem::RefreshAssetInformation()
 {
 	// Initialize display-related members
diff --git a/Engine/Source/Editor/SourceControlWindows/Private/SSourceControlCommon.h b/Engine/Source/Editor/SourceControlWindows/Private/SSourceControlCommon.h
--- a/Engine/Source/Editor/SourceControlWindows/Private/SSourceControlCommon.h
+++ b/Engine/Source/Editor/SourceControlWindows/Private/SSourceControlCommon.h
@@ -159,6 +159,15 @@ struct FFileTreeItem : public IChangelistTreeItem
 	/** Returns the asset type color of the item */
 	FSlateColor GetAssetTypeColor() const { return FSlateColor(AssetTypeColor); }
 
+	/** Returns the last modification time of the file/asset. */
+	FText GetLastModifiedTimestamp() const { return LastModifiedTimestamp; }
+	
+	/** Set the last time the files was saved on disk. */
+	void SetLastModifiedTimestamp(const FText& Timestamp) { LastModifiedTimestamp = Timestamp; }
+
+	/** Returns the user that checked out the file/asset (if any). */
+	FText GetCheckedOutByUser() const;
+
 	/** Returns the package name of the item to display */
 	FText GetPackageName() const { return PackageName; }
 
@@ -226,6 +235,9 @@ private:
 	/** Cached package name to display */
 	FText PackageName;
 
+	/** The timestamp of the last modification to the file. */
+	FText LastModifiedTimestamp;
+
 	/** Matching asset(s) to facilitate Locate in content browser */
 	FAssetDataArrayPtr Assets;
 
@@ -268,6 +280,8 @@ public:
 	const FText& GetDisplayPath() const { return AssetPath; }
 	const FText& GetDisplayType() const { return AssetType; }
 	const FSlateColor& GetDisplayColor() const { return AssetTypeColor; }
+	FText GetLastModifiedTimestamp() const { return LastModifiedTimestamp; }
+	void SetLastModifiedTimestamp(const FText& Timestamp) { LastModifiedTimestamp = Timestamp; }
 
 private:
 	TArray<FAssetData> Assets;
@@ -277,6 +291,7 @@ private:
 	FText AssetPath;
 	FText AssetType;
 	FSlateColor AssetTypeColor;
+	FText LastModifiedTimestamp;
 };
 
 
diff --git a/Engine/Source/Editor/SourceControlWindows/Private/SourceControlSettings.h b/Engine/Source/Editor/SourceControlWindows/Private/SourceControlSettings.h
new file mode 100644
--- /dev/null
+++ b/Engine/Source/Editor/SourceControlWindows/Private/SourceControlSettings.h
@@ -0,0 +1,25 @@
+// Copyright Epic Games, Inc. All Rights Reserved.
+
+#pragma once
+
+#include "UObject/Object.h"
+#include "SourceControlSettings.generated.h"
+
+/** Serializes source control window settings. */
+UCLASS(config=Editor)
+class USourceControlSettings : public UObject
+{
+	GENERATED_BODY()
+
+public:
+	USourceControlSettings() {}
+
+	UPROPERTY(config, EditAnywhere, Category="Source Control Changelist View")
+	bool bShowAssetTypeColumn = true;
+
+	UPROPERTY(config, EditAnywhere, Category="Source Control Changelist View")
+	bool bShowAssetLastModifiedTimeColumn = true;
+
+	UPROPERTY(config, EditAnywhere, Category="Source Control Changelist View")
+	bool bShowAssetCheckedOutByColumn = true;
+};
