Fixed a crash when refreshing the shelved files in the Changelist Window.

diff --git a/Engine/Source/Editor/SourceControlWindows/Private/SSourceControlChangelists.cpp b/Engine/Source/Editor/SourceControlWindows/Private/SSourceControlChangelists.cpp
--- a/Engine/Source/Editor/SourceControlWindows/Private/SSourceControlChangelists.cpp
+++ b/Engine/Source/Editor/SourceControlWindows/Private/SSourceControlChangelists.cpp
@@ -672,7 +672,7 @@ void SSourceControlChangelistsWidget::RequestFileStatusRefresh(const IChangelist
 	}
 	else if (Changelist.GetTreeItemType() == IChangelistTreeItem::ShelvedChangelist)
 	{
-		const FChangelistTreeItem& ChangelistItem = static_cast<const FChangelistTreeItem&>(Changelist);
+		const FChangelistTreeItem& ChangelistItem = static_cast<const FChangelistTreeItem&>(*Changelist.GetParent().Get());
 		Algo::Transform(ChangelistItem.ChangelistState->GetShelvedFilesStates(), Pathnames, [](const TSharedRef<ISourceControlState>& FileState) { return FileState->GetFilename(); });
 	}
 	else if (Changelist.GetTreeItemType() == IChangelistTreeItem::UncontrolledChangelist)
