Adds new uncontrolled changelist option to "Move Files To..." dialog.

#changelist validated

diff --git a/Engine/Source/Developer/UncontrolledChangelists/Private/UncontrolledChangelistsModule.cpp b/Engine/Source/Developer/UncontrolledChangelists/Private/UncontrolledChangelistsModule.cpp
--- a/Engine/Source/Developer/UncontrolledChangelists/Private/UncontrolledChangelistsModule.cpp
+++ b/Engine/Source/Developer/UncontrolledChangelists/Private/UncontrolledChangelistsModule.cpp
@@ -483,18 +483,20 @@ void FUncontrolledChangelistsModule::MoveFilesToControlledChangelist(const TArra
 	}
 }
 
-void FUncontrolledChangelistsModule::CreateUncontrolledChangelist(const FText& InDescription)
+TOptional<FUncontrolledChangelist> FUncontrolledChangelistsModule::CreateUncontrolledChangelist(const FText& InDescription)
 {
 	if (!IsEnabled())
 	{
-		return;
+		return TOptional<FUncontrolledChangelist>();
 	}
 
 	// Default constructor will generate a new GUID.
 	FUncontrolledChangelist NewUncontrolledChangelist;
-	UncontrolledChangelistsStateCache.Emplace(MoveTemp(NewUncontrolledChangelist), MakeShared<FUncontrolledChangelistState>(NewUncontrolledChangelist, InDescription));
+	UncontrolledChangelistsStateCache.Emplace(NewUncontrolledChangelist, MakeShared<FUncontrolledChangelistState>(NewUncontrolledChangelist, InDescription));
 
 	OnStateChanged();
+
+	return NewUncontrolledChangelist;
 }
 
 void FUncontrolledChangelistsModule::EditUncontrolledChangelist(const FUncontrolledChangelist& InUncontrolledChangelist, const FText& InNewDescription)
diff --git a/Engine/Source/Developer/UncontrolledChangelists/Public/UncontrolledChangelistsModule.h b/Engine/Source/Developer/UncontrolledChangelists/Public/UncontrolledChangelistsModule.h
--- a/Engine/Source/Developer/UncontrolledChangelists/Public/UncontrolledChangelistsModule.h
+++ b/Engine/Source/Developer/UncontrolledChangelists/Public/UncontrolledChangelistsModule.h
@@ -129,8 +129,9 @@ public:
 	/**
 	 * Creates a new Uncontrolled Changelist.
 	 * @param	InDescription	The description of the newly created Uncontrolled Changelist.
+	 * return	TOptional<FUncontrolledChangelist> set with the new Uncontrolled Changelist key if succeeded.
 	 */
-	void CreateUncontrolledChangelist(const FText& InDescription);
+	TOptional<FUncontrolledChangelist> CreateUncontrolledChangelist(const FText& InDescription);
 	
 	/**
      * Edits an Uncontrolled Changelist's description
diff --git a/Engine/Source/Editor/SourceControlWindows/Private/SSourceControlChangelists.cpp b/Engine/Source/Editor/SourceControlWindows/Private/SSourceControlChangelists.cpp
--- a/Engine/Source/Editor/SourceControlWindows/Private/SSourceControlChangelists.cpp
+++ b/Engine/Source/Editor/SourceControlWindows/Private/SSourceControlChangelists.cpp
@@ -2077,6 +2077,7 @@ void SSourceControlChangelistsWidget::OnMoveFiles()
 	}
 
 	const bool bAddNewChangelistEntry = true;
+	const bool bAddNewUncontrolledChangelistEntry = AreUncontrolledChangelistsEnabled();
 
 	// Build selection list for changelists
 	TArray<SSourceControlDescriptionItem> Items;
@@ -2091,6 +2092,15 @@ void SSourceControlChangelistsWidget::OnMoveFiles()
 			/*bCanEditDescription=*/true);
 	}
 
+	if (bAddNewUncontrolledChangelistEntry)
+	{
+		// Second item in the 'Move To' list is 'new uncontrolled changelist' (if enabled)
+		Items.Emplace(
+			LOCTEXT("SourceControl_NewUncontrolledChangelistText", "New Uncontrolled Changelist"),
+			LOCTEXT("SourceControl_NewUncontrolledChangelistDescription", "<enter description here>"),
+			/*bCanEditDescription=*/true);
+	}
+
 	const bool bCanEditAlreadyExistingChangelistDescription = false;
 
 	for (TSharedPtr<IChangelistTreeItem>& Changelist : ChangelistTreeNodes)
@@ -2154,11 +2164,39 @@ void SSourceControlChangelistsWidget::OnMoveFiles()
 				}
 			}));
 	}
+	else if ((bAddNewUncontrolledChangelistEntry && bAddNewChangelistEntry && PickedItem == 1) || (bAddNewUncontrolledChangelistEntry && !bAddNewChangelistEntry && PickedItem == 0)) // Move files to a new uncontrolled changelist
+	{
+		ExecuteUncontrolledChangelistOperation(LOCTEXT("Moving_Files_New_Uncontrolled_Changelist", "Moving file(s) to a new uncontrolled changelist..."), [&]()
+		{
+			TArray<FSourceControlStateRef> SelectedControlledFileStates;
+			TArray<FSourceControlStateRef> SelectedUnControlledFileStates;
+			FUncontrolledChangelistsModule& UncontrolledChangelistsModule = FUncontrolledChangelistsModule::Get();
+
+			GetSelectedFileStates(SelectedControlledFileStates, SelectedUnControlledFileStates);
+
+			TOptional<FUncontrolledChangelist> NewUncontrolledChangelist = UncontrolledChangelistsModule.CreateUncontrolledChangelist(ChangelistDescription);
+
+			if (!NewUncontrolledChangelist.IsSet())
+			{
+				SSourceControlCommon::DisplaySourceControlOperationNotification(LOCTEXT("Move_Files_New_Uncontrolled_Changelist_Failed", "Failed to create a new uncontrolled changelist."), SNotificationItem::CS_Fail);
+			}
+			else if (!SelectedControlledFileStates.IsEmpty() || !SelectedUnControlledFileStates.IsEmpty())
+			{
+				UncontrolledChangelistsModule.MoveFilesToUncontrolledChangelist(SelectedControlledFileStates, SelectedUnControlledFileStates, NewUncontrolledChangelist.GetValue());
+			}
+		});
+	}
 	else // Move files to an existing changelist or uncontrolled changelist.
 	{
 		// NOTE: The combo box indices are in this order: New changelist, existing changelist(s), existing uncontrolled changelist(s)
 		FChangelistTreeItemPtr MoveDestination;
-		const int32 ChangelistIndex = (bAddNewChangelistEntry ? PickedItem - 1 : PickedItem);
+		int32 ChangelistIndex = (bAddNewChangelistEntry ? PickedItem - 1 : PickedItem);
+
+		if (bAddNewUncontrolledChangelistEntry)
+		{
+			--ChangelistIndex;
+			check(ChangelistIndex >= 0);
+		}
 
 		if (ChangelistIndex < ChangelistTreeNodes.Num()) // Move files to a changelist
 		{
