Fix FOfflineFileTreeItem display name (Folders)
- Use same code to retrieve asset display information then FFileTreeItem

diff --git a/Engine/Source/Editor/SourceControlWindows/Private/SSourceControlCommon.cpp b/Engine/Source/Editor/SourceControlWindows/Private/SSourceControlCommon.cpp
--- a/Engine/Source/Editor/SourceControlWindows/Private/SSourceControlCommon.cpp
+++ b/Engine/Source/Editor/SourceControlWindows/Private/SSourceControlCommon.cpp
@@ -51,65 +51,65 @@ void IChangelistTreeItem::RemoveChild(const TSharedRef<IChangelistTreeItem>& Chi
 	}
 }
 
-//////////////////////////////////////////////////////////////////////////
-
-FFileTreeItem::FFileTreeItem(FSourceControlStateRef InFileState, bool bBeautifyPaths, bool bIsShelvedFile)
-	: IChangelistTreeItem(bIsShelvedFile ? IChangelistTreeItem::ShelvedFile : IChangelistTreeItem::File)
-	, FileState(InFileState)
-	, MinTimeBetweenUpdate(FTimespan::FromSeconds(5.f))
-	, LastUpdateTime()
-	, bAssetsUpToDate(false)
+static FString RetrieveAssetName(const FAssetData& InAssetData)
 {
-	CheckBoxState = ECheckBoxState::Checked;
-
-	// Initialize asset data first
+	static const FName NAME_ActorLabel(TEXT("ActorLabel"));
 
-	if (bBeautifyPaths)
+	if (InAssetData.FindTag(NAME_ActorLabel))
 	{
-		FSourceControlAssetDataCache& AssetDataCache = ISourceControlModule::Get().GetAssetDataCache();
-		bAssetsUpToDate = AssetDataCache.GetAssetDataArray(FileState, Assets);
+		FString ResultAssetName = TEXT("");
+
+		InAssetData.GetTagValue(NAME_ActorLabel, ResultAssetName);
+		return ResultAssetName;
 	}
-	else
+	else if (InAssetData.AssetClassPath == UActorFolder::StaticClass()->GetClassPathName())
 	{
-		// We do not need to wait for AssetData from the cache.
-		bAssetsUpToDate = true;
+		FString ActorFolderPath = UActorFolder::GetAssetRegistryInfoFromPackage(InAssetData.PackageName).GetDisplayName();
+		if (!ActorFolderPath.IsEmpty())
+		{
+			return ActorFolderPath;
+		}
 	}
 
-	RefreshAssetInformation();
+	return InAssetData.AssetName.ToString();
 }
 
-FText FFileTreeItem::GetCheckedOutByUser() const
+static FString RetrieveAssetPath(const FAssetData& InAssetData)
 {
-	FString Users;
-	if (FileState->IsCheckedOutOther(&Users))
+	int32 LastDot = -1;
+	FString Path = InAssetData.GetObjectPathString();
+
+	// Strip asset name from object path
+	if (Path.FindLastChar('.', LastDot))
 	{
-		return FText::FromString(Users);
+		Path.LeftInline(LastDot);
 	}
-	return FText::GetEmpty();
+
+	return Path;
 }
 
-void FFileTreeItem::RefreshAssetInformation()
+static void RefreshAssetInformationInternal(const TArray<FAssetData>& Assets, const FString& InFilename, FText& OutAssetName, FText& OutAssetPath, FText& OutAssetType, FText& OutPackageName, FColor& OutAssetTypeColor)
 {
 	// Initialize display-related members
-	FString Filename = FileState->GetFilename();
+	FString Filename = InFilename;
 	FString TempAssetName = SSourceControlCommon::GetDefaultAssetName().ToString();
 	FString TempAssetPath = Filename;
 	FString TempAssetType = SSourceControlCommon::GetDefaultAssetType().ToString();
 	FString TempPackageName = Filename;
 	FColor TempAssetColor = FColor(		// Copied from ContentBrowserCLR.cpp
-								   127 + FColor::Red.R / 2,	// Desaturate the colors a bit (GB colors were too.. much)
-								   127 + FColor::Red.G / 2,
-								   127 + FColor::Red.B / 2,
-								   200); // Opacity
+		127 + FColor::Red.R / 2,	// Desaturate the colors a bit (GB colors were too.. much)
+		127 + FColor::Red.G / 2,
+		127 + FColor::Red.B / 2,
+		200); // Opacity
 
-	if (Assets.IsValid() && (Assets->Num() > 0))
+	if (Assets.Num() > 0)
 	{
 		auto IsNotRedirector = [](const FAssetData& InAssetData) { return !InAssetData.IsRedirector(); };
-		int32 NumUserFacingAsset = Algo::CountIf(*Assets, IsNotRedirector);
+		int32 NumUserFacingAsset = Algo::CountIf(Assets, IsNotRedirector);
 
 		if (NumUserFacingAsset == 1)
 		{
-			const FAssetData& AssetData = *Algo::FindByPredicate(*Assets, IsNotRedirector);
+			const FAssetData& AssetData = *Algo::FindByPredicate(Assets, IsNotRedirector);
 
 			TempAssetName = RetrieveAssetName(AssetData);
 			TempAssetPath = RetrieveAssetPath(AssetData);
@@ -129,12 +129,12 @@ void FFileTreeItem::RefreshAssetInformation()
 		}
 		else
 		{
-			TempAssetName = RetrieveAssetName((*Assets)[0]);
-			TempAssetPath = RetrieveAssetPath((*Assets)[0]);
+			TempAssetName = RetrieveAssetName(Assets[0]);
+			TempAssetPath = RetrieveAssetPath(Assets[0]);
 
-			for (int32 i = 1; i < Assets->Num(); ++i)
+			for (int32 i = 1; i < Assets.Num(); ++i)
 			{
-				TempAssetName += TEXT(";") + RetrieveAssetName((*Assets)[i]);
+				TempAssetName += TEXT(";") + RetrieveAssetName(Assets[i]);
 			}
 
 			TempAssetType = SSourceControlCommon::GetDefaultMultipleAsset().ToString();
@@ -164,11 +164,55 @@ void FFileTreeItem::RefreshAssetInformation()
 	}
 
 	// Finally, assign the temp variables to the member variables
-	AssetName = FText::FromString(TempAssetName);
-	AssetPath = FText::FromString(TempAssetPath);
-	AssetType = FText::FromString(TempAssetType);
-	AssetTypeColor = TempAssetColor;
-	PackageName = FText::FromString(TempPackageName);
+	OutAssetName = FText::FromString(TempAssetName);
+	OutAssetPath = FText::FromString(TempAssetPath);
+	OutAssetType = FText::FromString(TempAssetType);
+	OutAssetTypeColor = TempAssetColor;
+	OutPackageName = FText::FromString(TempPackageName);
+}
+
+//////////////////////////////////////////////////////////////////////////
+
+FFileTreeItem::FFileTreeItem(FSourceControlStateRef InFileState, bool bBeautifyPaths, bool bIsShelvedFile)
+	: IChangelistTreeItem(bIsShelvedFile ? IChangelistTreeItem::ShelvedFile : IChangelistTreeItem::File)
+	, FileState(InFileState)
+	, MinTimeBetweenUpdate(FTimespan::FromSeconds(5.f))
+	, LastUpdateTime()
+	, bAssetsUpToDate(false)
+{
+	CheckBoxState = ECheckBoxState::Checked;
+
+	// Initialize asset data first
+
+	if (bBeautifyPaths)
+	{
+		FSourceControlAssetDataCache& AssetDataCache = ISourceControlModule::Get().GetAssetDataCache();
+		bAssetsUpToDate = AssetDataCache.GetAssetDataArray(FileState, Assets);
+	}
+	else
+	{
+		// We do not need to wait for AssetData from the cache.
+		bAssetsUpToDate = true;
+	}
+
+	RefreshAssetInformation();
+}
+
+FText FFileTreeItem::GetCheckedOutByUser() const
+{
+	FString Users;
+	if (FileState->IsCheckedOutOther(&Users))
+	{
+		return FText::FromString(Users);
+	}
+	return FText::GetEmpty();
+}
+
+void FFileTreeItem::RefreshAssetInformation()
+{
+	// Initialize display-related members
+	static TArray<FAssetData> NoAssets;
+	RefreshAssetInformationInternal(Assets.IsValid() ? *Assets : NoAssets, FileState->GetFilename(), AssetName, AssetPath, AssetType, PackageName, AssetTypeColor);
 }
 
 FText FFileTreeItem::GetAssetName() const
@@ -195,43 +239,6 @@ FText FFileTreeItem::GetAssetName()
 	return AssetName;
 }
 
-FString FFileTreeItem::RetrieveAssetName(const FAssetData& InAssetData) const
-{
-	static const FName NAME_ActorLabel(TEXT("ActorLabel"));
-
-	if (InAssetData.FindTag(NAME_ActorLabel))
-	{
-		FString ResultAssetName = TEXT("");
-		
-		InAssetData.GetTagValue(NAME_ActorLabel, ResultAssetName);
-		return ResultAssetName;
-	}
-	else if (InAssetData.AssetClassPath == UActorFolder::StaticClass()->GetClassPathName())
-	{
-		FString ActorFolderPath = UActorFolder::GetAssetRegistryInfoFromPackage(InAssetData.PackageName).GetDisplayName();
-		if (!ActorFolderPath.IsEmpty())
-		{
-			return ActorFolderPath;
-		}
-	}
-	
-	return InAssetData.AssetName.ToString();
-}
-
-FString FFileTreeItem::RetrieveAssetPath(const FAssetData& InAssetData) const
-{
-	int32 LastDot = -1;
-	FString Path = InAssetData.GetObjectPathString();
-
-	// Strip asset name from object path
-	if (Path.FindLastChar('.', LastDot))
-	{
-		Path.LeftInline(LastDot);
-	}
-
-	return Path;
-}
-
 //////////////////////////////////////////////////////////////////////////
 
 FText FShelvedChangelistTreeItem::GetDisplayText() const
@@ -255,57 +262,12 @@ FOfflineFileTreeItem::FOfflineFileTreeItem(const FString& InFilename)
 
 	USourceControlHelpers::GetAssetData(InFilename, Assets);
 
-	if (Assets.Num() > 0)
-	{
-		const FAssetData& AssetData = Assets[0];
-		AssetPath = FText::FromString(AssetData.GetObjectPathString());
-
-		// Find name, asset type & color only if there is exactly one asset
-		if (Assets.Num() == 1)
-		{
-			static FName NAME_ActorLabel(TEXT("ActorLabel"));
-			if (AssetData.FindTag(NAME_ActorLabel))
-			{
-				AssetData.GetTagValue(NAME_ActorLabel, AssetName);
-			}
-			else
-			{
-				AssetName = FText::FromName(AssetData.AssetName);
-			}
-
-			AssetType = FText::FromString(AssetData.AssetClassPath.ToString());
-
-			const FAssetToolsModule& AssetToolsModule = FModuleManager::LoadModuleChecked<FAssetToolsModule>(TEXT("AssetTools"));
-			const TSharedPtr<IAssetTypeActions> AssetTypeActions = AssetToolsModule.Get().GetAssetTypeActionsForClass(AssetData.GetClass()).Pin();
-			if (AssetTypeActions.IsValid())
-			{
-				AssetTypeColor = AssetTypeActions->GetTypeColor();
-			}
-			else
-			{
-				AssetTypeColor = FColor::White;
-			}
-		}
-		else
-		{
-			AssetType = SSourceControlCommon::GetDefaultMultipleAsset();
-			AssetTypeColor = FColor::White;
-		}
+	RefreshAssetInformation();
+}
 
-		// Beautify the package name
-		PackageName = AssetPath;
-	}
-	else if (FPackageName::TryConvertFilenameToLongPackageName(InFilename, TempString))
-	{
-		PackageName = FText::FromString(TempString);
-		// Fake asset name, asset path from the package name
-		AssetPath = PackageName;
-	}
-	else
-	{
-		AssetName = FText::FromString(FPaths::GetCleanFilename(InFilename));
-		AssetType = FText::Format(SSourceControlCommon::GetDefaultUnknownAssetType(), FText::FromString(FPaths::GetExtension(InFilename).ToUpper()));
-	}
+void FOfflineFileTreeItem::RefreshAssetInformation()
+{
+	RefreshAssetInformationInternal(Assets, Filename, AssetName, AssetPath, AssetType, PackageName, AssetTypeColor);
 }
 
 //////////////////////////////////////////////////////////////////////////
diff --git a/Engine/Source/Editor/SourceControlWindows/Private/SSourceControlCommon.h b/Engine/Source/Editor/SourceControlWindows/Private/SSourceControlCommon.h
--- a/Engine/Source/Editor/SourceControlWindows/Private/SSourceControlCommon.h
+++ b/Engine/Source/Editor/SourceControlWindows/Private/SSourceControlCommon.h
@@ -205,13 +205,6 @@ struct FFileTreeItem : public IChangelistTreeItem
 
 	bool IsShelved() const { return GetTreeItemType() == IChangelistTreeItem::ShelvedFile; }
 
-private:
-	/** Returns a string representing the name of the asset represented by the given AssetData */
-	FString RetrieveAssetName(const FAssetData& InAssetData) const;
-
-	/** Returns a string representing the path of the asset represented by the given AssetData */
-	FString RetrieveAssetPath(const FAssetData& InAssetData) const;
-
 public:
 	/** Shared pointer to the source control state object itself */
 	FSourceControlStateRef FileState;
@@ -273,13 +266,14 @@ struct FOfflineFileTreeItem : public IChangelistTreeItem
 {
 	explicit FOfflineFileTreeItem(const FString& InFilename);
 
+	void RefreshAssetInformation();
 public:
 	const FString& GetFilename() const { return Filename; }
 	const FText& GetPackageName() const { return PackageName; }
 	const FText& GetDisplayName() const { return AssetName; }
 	const FText& GetDisplayPath() const { return AssetPath; }
 	const FText& GetDisplayType() const { return AssetType; }
-	const FSlateColor& GetDisplayColor() const { return AssetTypeColor; }
+	const FColor& GetDisplayColor() const { return AssetTypeColor; }
 	FText GetLastModifiedTimestamp() const { return LastModifiedTimestamp; }
 	void SetLastModifiedTimestamp(const FText& Timestamp) { LastModifiedTimestamp = Timestamp; }
 
@@ -290,7 +284,7 @@ private:
 	FText AssetName;
 	FText AssetPath;
 	FText AssetType;
-	FSlateColor AssetTypeColor;
+	FColor AssetTypeColor;
 	FText LastModifiedTimestamp;
 };
 
