Instrumented source control upcate callbacks with with cpu scopes for monitoring performances.
Optimized Collection code invoked at each source control update to avoid comparing large set to detect if the in-memory collection is different from the collection on disk.

diff --git a/Engine/Source/Developer/CollectionManager/Private/Collection.cpp b/Engine/Source/Developer/CollectionManager/Private/Collection.cpp
--- a/Engine/Source/Developer/CollectionManager/Private/Collection.cpp
+++ b/Engine/Source/Developer/CollectionManager/Private/Collection.cpp
@@ -17,6 +17,7 @@
 #include "Misc/ScopeRWLock.h"
 #include "Async/ParallelFor.h"
 #include "String/ParseLines.h"
+#include "ProfilingDebugging/CpuProfilerTrace.h"
 
 #define LOCTEXT_NAMESPACE "CollectionManager"
 
@@ -156,6 +157,7 @@ bool FCollection::Load(FText& OutError)
 	}
 
 	DiskSnapshot.TakeSnapshot(*this);
+	bChangedSinceLastDiskSnapshot = false;
 
 	return true;
 }
@@ -282,6 +284,7 @@ bool FCollection::Save(const TArray<FText>& AdditionalChangelistText, FText& Out
 		FileVersion = ECollectionVersion::CurrentVersion;
 
 		DiskSnapshot.TakeSnapshot(*this);
+		bChangedSinceLastDiskSnapshot = false;
 	}
 
 	GWarn->EndSlowTask();
@@ -397,6 +400,7 @@ bool FCollection::DeleteSourceFile(FText& OutError)
 	if ( bSuccessfullyDeleted )
 	{
 		DiskSnapshot = FCollectionSnapshot();
+		bChangedSinceLastDiskSnapshot = (ObjectSet.Num() == 0);
 	}
 
 	return bSuccessfullyDeleted;
@@ -409,6 +413,7 @@ void FCollection::Empty()
 	DynamicQueryExpressionEvaluatorPtr.Reset();
 
 	DiskSnapshot.TakeSnapshot(*this);
+	bChangedSinceLastDiskSnapshot = false;
 }
 
 bool FCollection::AddObjectToCollection(const FSoftObjectPath& ObjectPath)
@@ -422,6 +427,7 @@ bool FCollection::AddObjectToCollection(const FSoftObjectPath& ObjectPath)
 	{
 		bool bAlreadyInSet = false;
 		ObjectSet.Add(ObjectPath, &bAlreadyInSet);
+		bChangedSinceLastDiskSnapshot |= !bAlreadyInSet;
 		return !bAlreadyInSet;
 	}
 
@@ -435,9 +441,10 @@ bool FCollection::RemoveObjectFromCollection(const FSoftObjectPath& ObjectPath)
 		return false;
 	}
 
-	if (StorageMode == ECollectionStorageMode::Static)
+	if (StorageMode == ECollectionStorageMode::Static && ObjectSet.Remove(ObjectPath) > 0)
 	{
-		return ObjectSet.Remove(ObjectPath) > 0;
+		bChangedSinceLastDiskSnapshot = true;
+		return true;
 	}
 
 	return false;
@@ -539,6 +546,8 @@ bool FCollection::TestDynamicQuery(const ITextFilterExpressionContext& InContext
 
 FCollectionStatusInfo FCollection::GetStatusInfo() const
 {
+	TRACE_CPUPROFILER_EVENT_SCOPE(FCollection::GetStatusInfo);
+
 	FCollectionStatusInfo StatusInfo;
 
 	StatusInfo.bIsDirty = IsDirty();
@@ -562,6 +571,8 @@ FCollectionStatusInfo FCollection::GetStatusInfo() const
 
 bool FCollection::IsDirty() const
 {
+	TRACE_CPUPROFILER_EVENT_SCOPE(FCollection::IsDirty);
+
 	if (ParentCollectionGuid != DiskSnapshot.ParentCollectionGuid)
 	{
 		return true;
@@ -572,22 +583,14 @@ bool FCollection::IsDirty() const
 		return true;
 	}
 
-	bool bHasChanges = false;
-
 	if (StorageMode == ECollectionStorageMode::Static)
 	{
-		TArray<FSoftObjectPath> ObjectsAdded;
-		TArray<FSoftObjectPath> ObjectsRemoved;
-		GetObjectDifferencesFromDisk(ObjectsAdded, ObjectsRemoved);
-
-		bHasChanges = ObjectsAdded.Num() != 0 || ObjectsRemoved.Num() != 0;
+		return bChangedSinceLastDiskSnapshot;
 	}
 	else
 	{
-		bHasChanges = DynamicQueryText != DiskSnapshot.DynamicQueryText;
+		return DynamicQueryText != DiskSnapshot.DynamicQueryText;
 	}
-
-	return bHasChanges;
 }
 
 bool FCollection::IsEmpty() const
@@ -736,6 +739,8 @@ bool FCollection::MergeWithCollection(const FCollection& Other)
 			{
 				ObjectSet.Remove(RemovedObjectName);
 			}
+			
+			bChangedSinceLastDiskSnapshot = true;
 		}
 	}
 	else
diff --git a/Engine/Source/Developer/CollectionManager/Private/Collection.h b/Engine/Source/Developer/CollectionManager/Private/Collection.h
--- a/Engine/Source/Developer/CollectionManager/Private/Collection.h
+++ b/Engine/Source/Developer/CollectionManager/Private/Collection.h
@@ -188,6 +188,9 @@ private:
 	/** Source control is used if true */
 	bool bUseSCC;
 
+	/** Indicates if the collection has changes in memory that weren't saved to disk. */
+	bool bChangedSinceLastDiskSnapshot = false;
+
 	/** The filename used to load this collection. Empty if it is new or never loaded from disk. */
 	FString SourceFilename;
 
diff --git a/Engine/Source/Developer/CollectionManager/Private/CollectionManager.cpp b/Engine/Source/Developer/CollectionManager/Private/CollectionManager.cpp
--- a/Engine/Source/Developer/CollectionManager/Private/CollectionManager.cpp
+++ b/Engine/Source/Developer/CollectionManager/Private/CollectionManager.cpp
@@ -1364,6 +1364,8 @@ bool FCollectionManager::UpdateCollection(FName CollectionName, ECollectionShare
 
 bool FCollectionManager::GetCollectionStatusInfo(FName CollectionName, ECollectionShareType::Type ShareType, FCollectionStatusInfo& OutStatusInfo) const
 {
+	TRACE_CPUPROFILER_EVENT_SCOPE(FCollectionManager::GetCollectionStatusInfo);
+
 	if (!ensure(ShareType < ECollectionShareType::CST_All))
 	{
 		// Bad share type
diff --git a/Engine/Source/Developer/SourceControl/Private/SourceControlFileStatusMonitor.cpp b/Engine/Source/Developer/SourceControl/Private/SourceControlFileStatusMonitor.cpp
--- a/Engine/Source/Developer/SourceControl/Private/SourceControlFileStatusMonitor.cpp
+++ b/Engine/Source/Developer/SourceControl/Private/SourceControlFileStatusMonitor.cpp
@@ -8,6 +8,7 @@
 #include "ISourceControlProvider.h"
 #include "Math/NumericLimits.h"
 #include "SourceControlOperations.h"
+#include "ProfilingDebugging/CpuProfilerTrace.h"
 
 FSourceControlFileStatusMonitor::FSourceControlFileStatusMonitor()
 : ProbationPeriodPolicy(FTimespan::FromSeconds(1))
@@ -257,6 +258,8 @@ bool FSourceControlFileStatusMonitor::Tick(float DeltaTime)
 
 void FSourceControlFileStatusMonitor::OnSourceControlStatusUpdate(const TSharedRef<ISourceControlOperation>& InOperation, ECommandResult::Type InResult)
 {
+	TRACE_CPUPROFILER_EVENT_SCOPE(FSourceControlFileStatusMonitor::OnSourceControlStatusUpdate);
+
 	check(IsInGameThread()); // Concurrency issues if the callback in invoked from a background thread.
 
 	double NowSecs = FPlatformTime::Seconds();
diff --git a/Engine/Source/Editor/ContentBrowser/Private/AssetViewWidgets.cpp b/Engine/Source/Editor/ContentBrowser/Private/AssetViewWidgets.cpp
--- a/Engine/Source/Editor/ContentBrowser/Private/AssetViewWidgets.cpp
+++ b/Engine/Source/Editor/ContentBrowser/Private/AssetViewWidgets.cpp
@@ -44,6 +44,7 @@
 #include "Misc/StringBuilder.h"
 #include "Modules/ModuleManager.h"
 #include "PluginDescriptor.h"
+#include "ProfilingDebugging/CpuProfilerTrace.h"
 #include "Rendering/SlateRenderer.h"
 #include "SAssetTagItem.h"
 #include "SThumbnailEditModeTools.h"
@@ -654,6 +655,8 @@ void SAssetViewItem::HandleSourceControlProviderChanged(ISourceControlProvider&
 
 void SAssetViewItem::HandleSourceControlStateChanged()
 {
+	TRACE_CPUPROFILER_EVENT_SCOPE(SAssetViewItem::HandleSourceControlStateChanged);
+
 	if (AssetItem && AssetItem->IsFile() && !AssetItem->IsTemporary() && ISourceControlModule::Get().IsEnabled())
 	{
 		FString AssetFilename;
diff --git a/Engine/Source/Editor/ContentBrowser/Private/SCollectionView.cpp b/Engine/Source/Editor/ContentBrowser/Private/SCollectionView.cpp
--- a/Engine/Source/Editor/ContentBrowser/Private/SCollectionView.cpp
+++ b/Engine/Source/Editor/ContentBrowser/Private/SCollectionView.cpp
@@ -43,6 +43,7 @@
 #include "Misc/ConfigCacheIni.h"
 #include "Misc/TextFilterUtils.h"
 #include "Modules/ModuleManager.h"
+#include "ProfilingDebugging/CpuProfilerTrace.h"
 #include "SAssetTagItem.h"
 #include "SlotBase.h"
 #include "SourcesData.h"
@@ -317,6 +318,8 @@ void SCollectionView::HandleSourceControlProviderChanged(ISourceControlProvider&
 
 void SCollectionView::HandleSourceControlStateChanged()
 {
+	TRACE_CPUPROFILER_EVENT_SCOPE(SCollectionView::HandleSourceControlStateChanged);
+
 	// Update the status of each collection
 	for (const auto& AvailableCollectionInfo : AvailableCollections)
 	{
@@ -326,6 +329,8 @@ void SCollectionView::HandleSourceControlStateChanged()
 
 void SCollectionView::UpdateCollectionItemStatus( const TSharedRef<FCollectionItem>& CollectionItem )
 {
+	TRACE_CPUPROFILER_EVENT_SCOPE(SCollectionView::UpdateCollectionItemStatus);
+
 	int32 NewObjectCount = 0;
 	TOptional<ECollectionItemStatus> NewStatus;
 
diff --git a/Engine/Source/Editor/LevelEditor/Private/SLevelEditorBuildAndSubmit.cpp b/Engine/Source/Editor/LevelEditor/Private/SLevelEditorBuildAndSubmit.cpp
--- a/Engine/Source/Editor/LevelEditor/Private/SLevelEditorBuildAndSubmit.cpp
+++ b/Engine/Source/Editor/LevelEditor/Private/SLevelEditorBuildAndSubmit.cpp
@@ -16,6 +16,7 @@
 #include "EditorBuildUtils.h"
 #include "Logging/MessageLog.h"
 #include "LightingBuildOptions.h"
+#include "ProfilingDebugging/CpuProfilerTrace.h"
 
 #include "Widgets/Docking/SDockTab.h"
 #include "Widgets/Layout/SExpandableArea.h"
@@ -298,6 +299,8 @@ void SLevelEditorBuildAndSubmit::OnEditorPackageModified(UPackage* Package)
 
 void SLevelEditorBuildAndSubmit::OnSourceControlStateChanged()
 {
+	TRACE_CPUPROFILER_EVENT_SCOPE(SLevelEditorBuildAndSubmit::OnSourceControlStateChanged);
+
 	if (bIsExtraPackagesSectionExpanded)
 	{
 		UpdatePackagesList();
diff --git a/Engine/Source/Editor/SceneOutliner/Private/SceneOutlinerTreeItemSCC.cpp b/Engine/Source/Editor/SceneOutliner/Private/SceneOutlinerTreeItemSCC.cpp
--- a/Engine/Source/Editor/SceneOutliner/Private/SceneOutlinerTreeItemSCC.cpp
+++ b/Engine/Source/Editor/SceneOutliner/Private/SceneOutlinerTreeItemSCC.cpp
@@ -7,6 +7,7 @@
 #include "ActorFolderTreeItem.h"
 #include "ISourceControlProvider.h"
 #include "ISourceControlModule.h"
+#include "ProfilingDebugging/CpuProfilerTrace.h"
 #include "WorldPartition/WorldPartitionActorDesc.h"
 
 FSceneOutlinerTreeItemSCC::FSceneOutlinerTreeItemSCC(FSceneOutlinerTreeItemPtr InTreeItemPtr)
@@ -122,6 +123,8 @@ void FSceneOutlinerTreeItemSCC::DisconnectSourceControl()
 
 void FSceneOutlinerTreeItemSCC::HandleSourceControlStateChanged(EStateCacheUsage::Type CacheUsage)
 {
+	TRACE_CPUPROFILER_EVENT_SCOPE(FSceneOutlinerTreeItemSCC::HandleSourceControlStateChanged);
+
 	FSourceControlStatePtr SourceControlState = ISourceControlModule::Get().GetProvider().GetState(ExternalPackageName, CacheUsage);
 	if (SourceControlState.IsValid())
 	{
diff --git a/Engine/Source/Editor/SourceControlWindows/Private/SSourceControlChangelists.cpp b/Engine/Source/Editor/SourceControlWindows/Private/SSourceControlChangelists.cpp
--- a/Engine/Source/Editor/SourceControlWindows/Private/SSourceControlChangelists.cpp
+++ b/Engine/Source/Editor/SourceControlWindows/Private/SSourceControlChangelists.cpp
@@ -43,6 +43,7 @@
 #include "Framework/Notifications/NotificationManager.h"
 #include "Framework/Docking/TabManager.h"
 #include "Framework/MultiBox/MultiBoxBuilder.h"
+#include "ProfilingDebugging/CpuProfilerTrace.h"
 
 #define LOCTEXT_NAMESPACE "SourceControlChangelist"
 
@@ -691,6 +692,8 @@ void SSourceControlChangelistsWidget::RequestFileStatusRefresh(const IChangelist
 
 void SSourceControlChangelistsWidget::RequestFileStatusRefresh(const TSet<FString>& PathnamesToMonitor)
 {
+	TRACE_CPUPROFILER_EVENT_SCOPE(SSourceControlChangelistsWidget::RequestFileStatusRefresh);
+
 	if (!ISourceControlModule::Get().IsEnabled())
 	{
 		return;
@@ -759,6 +762,8 @@ void SSourceControlChangelistsWidget::OnRefresh()
 		return;
 	}
 
+	TRACE_CPUPROFILER_EVENT_SCOPE(SSourceControlChangelistsWidget::Refresh);
+
 	// Views will be teared down and rebuilt from scratch, save the items that are expanded and/or selected to be able to restore those states after the rebuild.
 	FExpandedAndSelectionStates ExpandedAndSelectedStates;
 	SaveExpandedAndSelectionStates(ExpandedAndSelectedStates);
@@ -796,100 +801,108 @@ void SSourceControlChangelistsWidget::OnRefresh()
 	UncontrolledChangelistTreeNodes.Reset(UncontrolledChangelistStates.Num());
 	FileTreeNodes.Reset();
 
-	for (const TSharedRef<ISourceControlChangelistState>& ChangelistState : ChangelistsStates)
 	{
-		TSharedRef<IChangelistTreeItem> ChangelistNode = MakeShared<FChangelistTreeItem>(ChangelistState);
-		TSharedPtr<IChangelistTreeItem> ShelvedFilesNode;
-		bool bShelvedFilesNodeFilteredOut = true;
-		bool bSearchForcedChangelistNodeExpansion = false;
-		if (ChangelistState->GetShelvedFilesStates().Num() > 0)
-		{
-			ShelvedFilesNode = MakeShared<FShelvedChangelistTreeItem>();
-			bShelvedFilesNodeFilteredOut = !ChangelistTextFilter->PassesFilter(*ShelvedFilesNode);
-			bSearchForcedChangelistNodeExpansion = !bShelvedFilesNodeFilteredOut && !ChangelistTextFilter->GetRawFilterText().IsEmpty();
-		}
-		
-		// Check if the 'changelist' node must be kept.
-		if (!ChangelistTextFilter->PassesFilter(*ChangelistNode) && bShelvedFilesNodeFilteredOut)
-		{
-			continue; // Both the changelist and the shelved nodes are filtered out.
-		}
-
-		// Add the changelist.
-		ChangelistTreeNodes.Add(ChangelistNode);
-		if (bSearchForcedChangelistNodeExpansion)
-		{
-			ChangelistTreeView->SetItemExpansion(ChangelistNode, /*bShouldExpand*/true);
-		}
+		TRACE_CPUPROFILER_EVENT_SCOPE(SSourceControlChangelistsWidget::RefreshControlledChangelists);
 
-		// Add the changelist files.
-		for (const TSharedRef<ISourceControlState>& FileState : ChangelistState->GetFilesStates())
+		for (const TSharedRef<ISourceControlChangelistState>& ChangelistState : ChangelistsStates)
 		{
-			TSharedRef<FFileTreeItem> FileItem = MakeShared<FFileTreeItem>(FileState, bBeautifyPaths);
-			if (FileTextFilter->PassesFilter(*FileItem))
+			TSharedRef<IChangelistTreeItem> ChangelistNode = MakeShared<FChangelistTreeItem>(ChangelistState);
+			TSharedPtr<IChangelistTreeItem> ShelvedFilesNode;
+			bool bShelvedFilesNodeFilteredOut = true;
+			bool bSearchForcedChangelistNodeExpansion = false;
+			if (ChangelistState->GetShelvedFilesStates().Num() > 0)
 			{
-				ChangelistNode->AddChild(FileItem);
-				SlowTask.EnterProgressFrame();
-				bBeautifyPaths &= !SlowTask.ShouldCancel();
+				ShelvedFilesNode = MakeShared<FShelvedChangelistTreeItem>();
+				bShelvedFilesNodeFilteredOut = !ChangelistTextFilter->PassesFilter(*ShelvedFilesNode);
+				bSearchForcedChangelistNodeExpansion = !bShelvedFilesNodeFilteredOut && !ChangelistTextFilter->GetRawFilterText().IsEmpty();
+			}
+		
+			// Check if the 'changelist' node must be kept.
+			if (!ChangelistTextFilter->PassesFilter(*ChangelistNode) && bShelvedFilesNodeFilteredOut)
+			{
+				continue; // Both the changelist and the shelved nodes are filtered out.
 			}
-		}
 
-		if (ShelvedFilesNode)
-		{
-			// Add a shelved files node under the changelist node.
-			ChangelistNode->AddChild(ShelvedFilesNode.ToSharedRef());
+			// Add the changelist.
+			ChangelistTreeNodes.Add(ChangelistNode);
+			if (bSearchForcedChangelistNodeExpansion)
+			{
+				ChangelistTreeView->SetItemExpansion(ChangelistNode, /*bShouldExpand*/true);
+			}
 
-			// Add the shelved files.
-			for (const TSharedRef<ISourceControlState>& ShelvedFileState : ChangelistState->GetShelvedFilesStates())
+			// Add the changelist files.
+			for (const TSharedRef<ISourceControlState>& FileState : ChangelistState->GetFilesStates())
 			{
-				TSharedRef<FShelvedFileTreeItem> ShelvedFileItem = MakeShared<FShelvedFileTreeItem>(ShelvedFileState, bBeautifyPaths);
-				if (FileTextFilter->PassesFilter(*ShelvedFileItem))
+				TSharedRef<FFileTreeItem> FileItem = MakeShared<FFileTreeItem>(FileState, bBeautifyPaths);
+				if (FileTextFilter->PassesFilter(*FileItem))
 				{
-					ShelvedFilesNode->AddChild(ShelvedFileItem);
+					ChangelistNode->AddChild(FileItem);
 					SlowTask.EnterProgressFrame();
 					bBeautifyPaths &= !SlowTask.ShouldCancel();
 				}
 			}
-		}
 
-		SlowTask.EnterProgressFrame();
-		bBeautifyPaths &= !SlowTask.ShouldCancel();
+			if (ShelvedFilesNode)
+			{
+				// Add a shelved files node under the changelist node.
+				ChangelistNode->AddChild(ShelvedFilesNode.ToSharedRef());
+
+				// Add the shelved files.
+				for (const TSharedRef<ISourceControlState>& ShelvedFileState : ChangelistState->GetShelvedFilesStates())
+				{
+					TSharedRef<FShelvedFileTreeItem> ShelvedFileItem = MakeShared<FShelvedFileTreeItem>(ShelvedFileState, bBeautifyPaths);
+					if (FileTextFilter->PassesFilter(*ShelvedFileItem))
+					{
+						ShelvedFilesNode->AddChild(ShelvedFileItem);
+						SlowTask.EnterProgressFrame();
+						bBeautifyPaths &= !SlowTask.ShouldCancel();
+					}
+				}
+			}
+
+			SlowTask.EnterProgressFrame();
+			bBeautifyPaths &= !SlowTask.ShouldCancel();
+		}
 	}
 
-	for (const TSharedRef<FUncontrolledChangelistState>& UncontrolledChangelistState : UncontrolledChangelistStates)
 	{
-		// Add an uncontrolled changelist.
-		TSharedRef<IChangelistTreeItem> UncontrolledChangelistNode = MakeShared<FUncontrolledChangelistTreeItem>(UncontrolledChangelistState);
-		if (!UncontrolledChangelistTextFilter->PassesFilter(*UncontrolledChangelistNode))
-		{
-			continue;
-		}
-		UncontrolledChangelistTreeNodes.Add(UncontrolledChangelistNode);
+		TRACE_CPUPROFILER_EVENT_SCOPE(SSourceControlChangelistsWidget::RefreshUncontrolledChangelists);
 
-		for (const TSharedRef<ISourceControlState>& FileState : UncontrolledChangelistState->GetFilesStates())
+		for (const TSharedRef<FUncontrolledChangelistState>& UncontrolledChangelistState : UncontrolledChangelistStates)
 		{
-			TSharedRef<FFileTreeItem> UncontrolledFileItem = MakeShared<FFileTreeItem>(FileState, bBeautifyPaths);
-			if (FileTextFilter->PassesFilter(*UncontrolledFileItem))
+			// Add an uncontrolled changelist.
+			TSharedRef<IChangelistTreeItem> UncontrolledChangelistNode = MakeShared<FUncontrolledChangelistTreeItem>(UncontrolledChangelistState);
+			if (!UncontrolledChangelistTextFilter->PassesFilter(*UncontrolledChangelistNode))
 			{
-				UncontrolledChangelistNode->AddChild(UncontrolledFileItem);
-				SlowTask.EnterProgressFrame();
-				bBeautifyPaths &= !SlowTask.ShouldCancel();
+				continue;
 			}
-		}
+			UncontrolledChangelistTreeNodes.Add(UncontrolledChangelistNode);
 
-		for (const FString& Filename : UncontrolledChangelistState->GetOfflineFiles())
-		{
-			TSharedRef<FOfflineFileTreeItem> OfflineFileItem = MakeShared<FOfflineFileTreeItem>(Filename);
-			if (FileTextFilter->PassesFilter(*OfflineFileItem))
+			for (const TSharedRef<ISourceControlState>& FileState : UncontrolledChangelistState->GetFilesStates())
+			{
+				TSharedRef<FFileTreeItem> UncontrolledFileItem = MakeShared<FFileTreeItem>(FileState, bBeautifyPaths);
+				if (FileTextFilter->PassesFilter(*UncontrolledFileItem))
+				{
+					UncontrolledChangelistNode->AddChild(UncontrolledFileItem);
+					SlowTask.EnterProgressFrame();
+					bBeautifyPaths &= !SlowTask.ShouldCancel();
+				}
+			}
+
+			for (const FString& Filename : UncontrolledChangelistState->GetOfflineFiles())
 			{
-				UncontrolledChangelistNode->AddChild(OfflineFileItem);
-				SlowTask.EnterProgressFrame();
-				bBeautifyPaths &= !SlowTask.ShouldCancel();
+				TSharedRef<FOfflineFileTreeItem> OfflineFileItem = MakeShared<FOfflineFileTreeItem>(Filename);
+				if (FileTextFilter->PassesFilter(*OfflineFileItem))
+				{
+					UncontrolledChangelistNode->AddChild(OfflineFileItem);
+					SlowTask.EnterProgressFrame();
+					bBeautifyPaths &= !SlowTask.ShouldCancel();
+				}
 			}
-		}
 
-		SlowTask.EnterProgressFrame();
-		bBeautifyPaths &= !SlowTask.ShouldCancel();
+			SlowTask.EnterProgressFrame();
+			bBeautifyPaths &= !SlowTask.ShouldCancel();
+		}
 	}
 
 	// Views were rebuilt from scratch, try expanding and selecting the nodes that were in that state before the update.
@@ -920,6 +933,8 @@ void SSourceControlChangelistsWidget::OnSourceControlProviderChanged(ISourceCont
 
 void SSourceControlChangelistsWidget::OnSourceControlStateChanged()
 {
+	TRACE_CPUPROFILER_EVENT_SCOPE(SSourceControlChangelistsWidget::OnSourceControlStateChanged);
+
 	// NOTE: No need to call RequestChangelistsRefresh() to force the SCC to update internal states. We are being invoked because it was update, we just
 	//       need to update the UI to reflect those state changes.
 	OnRefresh();
