Optimized source control changelist sorting
  - Removed conversion/creation of temporary strings when sorting.
  - Removed several branching operations.
  - Used Stricmp rather than operator<() as it tells if the string is <> or == in a single operation.
  - Sorted the 'Last Modified' column using FDateTime rather that the FText/FString.
  - Mesured before/after the change with Insight, in average, the sorting 15 now times faster. (avg 2.4 ms for 7600 items).
 - Do not sort by default, waited until the user decide to sort to enable the feature.

diff --git a/Engine/Source/Editor/SourceControlWindows/Private/SSourceControlChangelists.cpp b/Engine/Source/Editor/SourceControlWindows/Private/SSourceControlChangelists.cpp
--- a/Engine/Source/Editor/SourceControlWindows/Private/SSourceControlChangelists.cpp
+++ b/Engine/Source/Editor/SourceControlWindows/Private/SSourceControlChangelists.cpp
@@ -31,6 +31,7 @@
 #include "SourceControlMenuContext.h"
 #include "SourceControlSettings.h"
 #include "AssetToolsModule.h"
+#include "Misc/CString.h"
 #include "Misc/MessageDialog.h"
 #include "Misc/ScopedSlowTask.h"
 #include "Algo/AnyOf.h"
@@ -106,15 +107,6 @@ FText UpdateChangelistDescriptionToSubmitIfNeeded(const bool bInValidationResult
 
 void UpdateLastModifiedTimestamp(IChangelistTreeItem& Item)
 {
-	auto GetTimestampText = [](const FString& Pathname)
-	{
-		if (IFileManager::Get().FileExists(*Pathname))
-		{
-			return FText::AsDateTime(IFileManager::Get().GetTimeStamp(*Pathname), EDateTimeStyle::Short);
-		}
-		return FText::GetEmpty();
-	};
-
 	switch (Item.GetTreeItemType())
 	{
 		// NOTE: We may not want to display the last saved timestamp for shelved files as this could lead users to think the file
@@ -123,7 +115,7 @@ void UpdateLastModifiedTimestamp(IChangelistTreeItem& Item)
 		case IChangelistTreeItem::ShelvedFile:
 		{
 			FFileTreeItem& FileItem = static_cast<FFileTreeItem&>(Item);
-			FileItem.SetLastModifiedTimestamp(GetTimestampText(FileItem.GetFileName().ToString()));
+			FileItem.SetLastModifiedDateTime(IFileManager::Get().GetTimeStamp(*FileItem.GetFileName().ToString()));
 			return;
 		}
 		break;
@@ -131,7 +123,7 @@ void UpdateLastModifiedTimestamp(IChangelistTreeItem& Item)
 		case IChangelistTreeItem::OfflineFile:
 		{
 			FOfflineFileTreeItem& OfflineItem = static_cast<FOfflineFileTreeItem&>(Item);
-			OfflineItem.SetLastModifiedTimestamp(GetTimestampText(OfflineItem.GetFilename()));
+			OfflineItem.SetLastModifiedDateTime(IFileManager::Get().GetTimeStamp(*OfflineItem.GetFilename()));
 			return;
 		}
 
@@ -353,7 +345,8 @@ void SSourceControlChangelistsWidget::Construct(const FArguments& InArgs)
 	SourceControlStateChangedDelegateHandle = SCCModule.GetProvider().RegisterSourceControlStateChanged_Handle(FSourceControlStateChanged::FDelegate::CreateSP(this, &SSourceControlChangelistsWidget::OnSourceControlStateChanged));
 	UncontrolledChangelistModule.OnUncontrolledChangelistModuleChanged.AddSP(this, &SSourceControlChangelistsWidget::OnSourceControlStateChanged);
 
-	PrimarySortedColumn = SourceControlFileViewColumn::Name::Id();
+	// No default sorting, sorting make the view slower, just pay the cost if user wants it.
+	PrimarySortedColumn = FName();
 
 	ChangelistTreeView = CreateChangelistTreeView(ChangelistTreeNodes);
 	UncontrolledChangelistTreeView = CreateChangelistTreeView(UncontrolledChangelistTreeNodes);
@@ -2899,162 +2892,136 @@ void SSourceControlChangelistsWidget::OnColumnSortModeChanged(const EColumnSortP
 
 void SSourceControlChangelistsWidget::SortFileView()
 {
-	// Invoked when sorting the icons. This gives a priority to the status for sorting purpose.
-	auto GetSourceControlStateWeight = [](const ISourceControlState& State)
-	{
-		if (!State.IsCurrent())        { return 0; } // First if sorted in ascending order.
-		if (State.IsUnknown())         { return 1; }
-		if (State.IsConflicted())      { return 2; }
-		if (State.IsCheckedOutOther()) { return 3; }
-		if (State.IsCheckedOut())      { return 4; }
-		if (State.IsDeleted())         { return 5; }
-		if (State.IsAdded())           { return 6; }
-		else                           { return 7; }
-	};
+	TRACE_CPUPROFILER_EVENT_SCOPE(SSourceControlChangelistsWidget::SortFileView);
 
-	// Implements the equivalent of operator< to compare the files status.
-	auto OperatorLessFileIcon = [&GetSourceControlStateWeight](const FFileTreeItem& Lhs, const FFileTreeItem& Rhs) -> bool
+	if (PrimarySortedColumn.IsNone() || FileTreeNodes.IsEmpty())
 	{
-		return GetSourceControlStateWeight(Lhs.FileState.Get()) < GetSourceControlStateWeight(Rhs.FileState.Get());
-	};
+		return; // No column selected for sorting or nothing to sort.
+	}
 
-	auto OperatorLessOfflineFileIcon = [](const FOfflineFileTreeItem& Lhs, const FOfflineFileTreeItem& Rhs) -> bool
+	auto CompareIcons = [](const IFileViewTreeItem* Lhs, const IFileViewTreeItem* Rhs)
 	{
-		return false; // Offline file do not have state, don't change the order.
+		int32 LhsVal = Lhs->GetIconSortingPriority();
+		int32 RhsVal = Rhs->GetIconSortingPriority();
+		return LhsVal < RhsVal ? -1 : ( LhsVal == RhsVal ? 0 : 1);
 	};
 
-	auto OperatorLessIcon = [&OperatorLessFileIcon, &OperatorLessOfflineFileIcon](IChangelistTreeItem* Lhs, IChangelistTreeItem* Rhs)-> bool
+	auto CompareNames = [](const IFileViewTreeItem* Lhs, const IFileViewTreeItem* Rhs)
 	{
-		check(Lhs->GetTreeItemType() == Rhs->GetTreeItemType());
-		switch (Lhs->GetTreeItemType())
-		{
-		case IChangelistTreeItem::File:
-		case IChangelistTreeItem::ShelvedFile:
-			return OperatorLessFileIcon(*static_cast<FFileTreeItem*>(Lhs), *static_cast<FFileTreeItem*>(Rhs));
-		case IChangelistTreeItem::OfflineFile:
-			return OperatorLessOfflineFileIcon(*static_cast<FOfflineFileTreeItem*>(Lhs), *static_cast<FOfflineFileTreeItem*>(Rhs));
-		default:
-			return false;
-		}
+		return FCString::Stricmp(*Lhs->GetName(), *Rhs->GetName());
 	};
 
-	auto GetName = [](IChangelistTreeItem* Item) -> FString
+	auto ComparePaths = [](const IFileViewTreeItem* Lhs, const IFileViewTreeItem* Rhs)
 	{
-		switch (Item->GetTreeItemType())
-		{
-		case IChangelistTreeItem::File:
-		case IChangelistTreeItem::ShelvedFile:
-			return static_cast<FFileTreeItem*>(Item)->GetAssetName().ToString();
-		case IChangelistTreeItem::OfflineFile:
-			return static_cast<FOfflineFileTreeItem*>(Item)->GetDisplayName().ToString();
-		default:
-			return FString();
-		}
+		return FCString::Stricmp(*Lhs->GetPath(), *Rhs->GetPath());
 	};
 
-	auto GetPath = [](IChangelistTreeItem* Item) -> FString
+	auto CompareTypes = [](const IFileViewTreeItem* Lhs, const IFileViewTreeItem* Rhs)
 	{
-		switch (Item->GetTreeItemType())
-		{
-		case IChangelistTreeItem::File:
-		case IChangelistTreeItem::ShelvedFile:
-			return static_cast<FFileTreeItem*>(Item)->GetAssetPath().ToString();
-		case IChangelistTreeItem::OfflineFile:
-			return static_cast<FOfflineFileTreeItem*>(Item)->GetDisplayPath().ToString();
-		default:
-			return FString();
-		}
+		return FCString::Stricmp(*Lhs->GetType(), *Rhs->GetType());
 	};
 
-	auto GetType = [](IChangelistTreeItem* Item)  -> FString
+	auto CompareLastModified = [](const IFileViewTreeItem* Lhs, const IFileViewTreeItem* Rhs)
 	{
-		switch (Item->GetTreeItemType())
-		{
-		case IChangelistTreeItem::File:
-		case IChangelistTreeItem::ShelvedFile:
-			return static_cast<FFileTreeItem*>(Item)->GetAssetType().ToString();
-		case IChangelistTreeItem::OfflineFile:
-			return static_cast<FOfflineFileTreeItem*>(Item)->GetDisplayType().ToString();
-		default:
-			return FString();
-		}
+		const FDateTime& LhsVal = Lhs->GetLastModifiedDateTime();
+		const FDateTime& RhsVal = Rhs->GetLastModifiedDateTime();
+		return LhsVal < RhsVal ? -1 : ( LhsVal == RhsVal ? 0 : 1);
 	};
 
-	auto GetModifiedTimestamp = [](IChangelistTreeItem* Item)  -> FString
+	auto CompareCheckedOutBy = [](const IFileViewTreeItem* Lhs, const IFileViewTreeItem* Rhs)
 	{
-		switch (Item->GetTreeItemType())
-		{
-		case IChangelistTreeItem::File:
-		case IChangelistTreeItem::ShelvedFile:
-			return static_cast<FFileTreeItem*>(Item)->GetLastModifiedTimestamp().ToString();
-		case IChangelistTreeItem::OfflineFile:
-			return static_cast<FOfflineFileTreeItem*>(Item)->GetLastModifiedTimestamp().ToString();
-		default:
-			return FString();
-		}
+		return FCString::Stricmp(*Lhs->GetCheckedOutBy(), *Rhs->GetCheckedOutBy());
 	};
 
-	auto GetCheckedOutByUser = [](IChangelistTreeItem* Item)  -> FString
-	{
-		switch (Item->GetTreeItemType())
-		{
-		case IChangelistTreeItem::File:
-		case IChangelistTreeItem::ShelvedFile:
-			return static_cast<FFileTreeItem*>(Item)->GetCheckedOutByUser().ToString();
-		default:
-			return FString();
-		}
-	};
-
-	auto Compare = [&](const TSharedPtr<IChangelistTreeItem>& Lhs, const TSharedPtr<IChangelistTreeItem>& Rhs, const FName& ColumnId, EColumnSortMode::Type SortMode) -> bool
+	auto GetCompareFunc = [&](const FName& ColumnId)
 	{
 		if (ColumnId == SourceControlFileViewColumn::Icon::Id())
 		{
-			return SortMode == EColumnSortMode::Ascending ? OperatorLessIcon(Lhs.Get(), Rhs.Get()) : OperatorLessIcon(Rhs.Get(), Lhs.Get());
+			return TFunction<int32(const IFileViewTreeItem*, const IFileViewTreeItem*)>(CompareIcons);
 		}
 		else if (ColumnId == SourceControlFileViewColumn::Name::Id())
 		{
-			return SortMode == EColumnSortMode::Ascending ? GetName(Lhs.Get()) < GetName(Rhs.Get()) : GetName(Lhs.Get()) > GetName(Rhs.Get());
+			return TFunction<int32(const IFileViewTreeItem*, const IFileViewTreeItem*)>(CompareNames);
 		}
 		else if (ColumnId == SourceControlFileViewColumn::Path::Id())
 		{
-			return SortMode == EColumnSortMode::Ascending ? GetPath(Lhs.Get()) < GetPath(Rhs.Get()) : GetPath(Lhs.Get()) > GetPath(Rhs.Get());
+			return TFunction<int32(const IFileViewTreeItem*, const IFileViewTreeItem*)>(ComparePaths);
 		}
 		else if (ColumnId == SourceControlFileViewColumn::Type::Id())
 		{
-			return SortMode == EColumnSortMode::Ascending ? GetType(Lhs.Get()) < GetType(Rhs.Get()) : GetType(Lhs.Get()) > GetType(Rhs.Get());
+			return TFunction<int32(const IFileViewTreeItem*, const IFileViewTreeItem*)>(CompareTypes);
 		}
 		else if (ColumnId == SourceControlFileViewColumn::LastModifiedTimestamp::Id())
 		{
-			return SortMode == EColumnSortMode::Ascending ? GetModifiedTimestamp(Lhs.Get()) < GetModifiedTimestamp(Rhs.Get()) : GetModifiedTimestamp(Lhs.Get()) > GetModifiedTimestamp(Rhs.Get());
+			return TFunction<int32(const IFileViewTreeItem*, const IFileViewTreeItem*)>(CompareLastModified);
 		}
 		else if (ColumnId == SourceControlFileViewColumn::CheckedOutByUser::Id())
 		{
-			return SortMode == EColumnSortMode::Ascending ? GetCheckedOutByUser(Lhs.Get()) < GetCheckedOutByUser(Rhs.Get()) : GetCheckedOutByUser(Lhs.Get()) > GetCheckedOutByUser(Rhs.Get());
+			return TFunction<int32(const IFileViewTreeItem*, const IFileViewTreeItem*)>(CompareCheckedOutBy);
 		}
 		else
 		{
 			checkNoEntry();
-			return false;
-		}
+			return TFunction<int32(const IFileViewTreeItem*, const IFileViewTreeItem*)>();
+		};
 	};
 
-	FileTreeNodes.Sort([&](const TSharedPtr<IChangelistTreeItem>& Lhs, const TSharedPtr<IChangelistTreeItem>& Rhs)
+	TFunction<int32(const IFileViewTreeItem*, const IFileViewTreeItem*)> PrimaryCompare = GetCompareFunc(PrimarySortedColumn);
+	TFunction<int32(const IFileViewTreeItem*, const IFileViewTreeItem*)> SecondaryCompare;
+	if (!SecondarySortedColumn.IsNone())
 	{
-		if (Compare(Lhs, Rhs, PrimarySortedColumn, PrimarySortMode))
-		{
-			return true; // Lhs must be before Rhs based on the primary sort order.
-		}
-		else if (Compare(Rhs, Lhs, PrimarySortedColumn, PrimarySortMode)) // Invert operands order (goal is to check if operands are equal or not)
+		SecondaryCompare = GetCompareFunc(SecondarySortedColumn);
+	}
+
+	if (PrimarySortMode == EColumnSortMode::Ascending)
+	{
+		// NOTE: StableSort() would give a better experience when the sorted columns(s) has the same values and new values gets added, but it is slower
+		//       with large changelists (7600 items was about 1.8x slower in average measured with Unreal Insight). Because this code runs in the main
+		//       thread and can be invoked a lot, the trade off went if favor of speed.
+		FileTreeNodes.Sort([this, &PrimaryCompare, &SecondaryCompare](const TSharedPtr<IChangelistTreeItem>& Lhs, const TSharedPtr<IChangelistTreeItem>& Rhs)
 		{
-			return false; // Rhs must be before Lhs based on the primary sort.
-		}
-		else if (!SecondarySortedColumn.IsNone()) // Lhs == Rhs on the primary column, need to sort according the secondary column if one is set.
+			int32 Result = PrimaryCompare(static_cast<IFileViewTreeItem*>(Lhs.Get()), static_cast<IFileViewTreeItem*>(Rhs.Get()));
+			if (Result < 0)
+			{
+				return true;
+			}
+			else if (Result > 0 || !SecondaryCompare)
+			{
+				return false;
+			}
+			else if (SecondarySortMode == EColumnSortMode::Ascending)
+			{
+				return SecondaryCompare(static_cast<IFileViewTreeItem*>(Lhs.Get()), static_cast<IFileViewTreeItem*>(Rhs.Get())) < 0;
+			}
+			else
+			{
+				return SecondaryCompare(static_cast<IFileViewTreeItem*>(Lhs.Get()), static_cast<IFileViewTreeItem*>(Rhs.Get())) > 0;
+			}
+		});
+	}
+	else
+	{
+		FileTreeNodes.Sort([this, &PrimaryCompare, &SecondaryCompare](const TSharedPtr<IChangelistTreeItem>& Lhs, const TSharedPtr<IChangelistTreeItem>& Rhs)
 		{
-			return Compare(Lhs, Rhs, SecondarySortedColumn, SecondarySortMode);
-		}
-		return false;
-	});
+			int32 Result = PrimaryCompare(static_cast<IFileViewTreeItem*>(Lhs.Get()), static_cast<IFileViewTreeItem*>(Rhs.Get()));
+			if (Result > 0)
+			{
+				return true;
+			}
+			else if (Result < 0 || !SecondaryCompare)
+			{
+				return false;
+			}
+			else if (SecondarySortMode == EColumnSortMode::Ascending)
+			{
+				return SecondaryCompare(static_cast<IFileViewTreeItem*>(Lhs.Get()), static_cast<IFileViewTreeItem*>(Rhs.Get())) < 0;
+			}
+			else
+			{
+				return SecondaryCompare(static_cast<IFileViewTreeItem*>(Lhs.Get()), static_cast<IFileViewTreeItem*>(Rhs.Get())) > 0;
+			}
+		});
+	}
 }
 
 void SSourceControlChangelistsWidget::OnChangelistSearchTextChanged(const FText& InFilterText)
diff --git a/Engine/Source/Editor/SourceControlWindows/Private/SSourceControlCommon.cpp b/Engine/Source/Editor/SourceControlWindows/Private/SSourceControlCommon.cpp
--- a/Engine/Source/Editor/SourceControlWindows/Private/SSourceControlCommon.cpp
+++ b/Engine/Source/Editor/SourceControlWindows/Private/SSourceControlCommon.cpp
@@ -88,7 +88,7 @@ static FString RetrieveAssetPath(const FAssetData& InAssetData)
 	return Path;
 }
 
-static void RefreshAssetInformationInternal(const TArray<FAssetData>& Assets, const FString& InFilename, FText& OutAssetName, FText& OutAssetPath, FText& OutAssetType, FText& OutPackageName, FColor& OutAssetTypeColor)
+static void RefreshAssetInformationInternal(const TArray<FAssetData>& Assets, const FString& InFilename, FString& OutAssetName, FString& OutAssetPath, FString& OutAssetType, FText& OutPackageName, FColor& OutAssetTypeColor)
 {
 	// Initialize display-related members
 	FString Filename = InFilename;
@@ -164,17 +164,35 @@ static void RefreshAssetInformationInternal(const TArray<FAssetData>& Assets, co
 	}
 
 	// Finally, assign the temp variables to the member variables
-	OutAssetName = FText::FromString(TempAssetName);
-	OutAssetPath = FText::FromString(TempAssetPath);
-	OutAssetType = FText::FromString(TempAssetType);
+	OutAssetName = TempAssetName;
+	OutAssetPath = TempAssetPath;
+	OutAssetType = TempAssetType;
 	OutAssetTypeColor = TempAssetColor;
 	OutPackageName = FText::FromString(TempPackageName);
 }
 
 //////////////////////////////////////////////////////////////////////////
 
+FString IFileViewTreeItem::DefaultStrValue; // Default is an empty string.
+FDateTime IFileViewTreeItem::DefaultDateTimeValue; // Default is FDateTime::MinValue().
+
+void IFileViewTreeItem::SetLastModifiedDateTime(const FDateTime& Timestamp)
+{
+	LastModifiedDateTime = Timestamp;
+	if (Timestamp != FDateTime::MinValue())
+	{
+		LastModifiedTimestampText = FText::AsDateTime(Timestamp, EDateTimeStyle::Short);
+	}
+	else
+	{
+		LastModifiedTimestampText = FText::GetEmpty();
+	}
+}
+
+//////////////////////////////////////////////////////////////////////////
+
 FFileTreeItem::FFileTreeItem(FSourceControlStateRef InFileState, bool bBeautifyPaths, bool bIsShelvedFile)
-	: IChangelistTreeItem(bIsShelvedFile ? IChangelistTreeItem::ShelvedFile : IChangelistTreeItem::File)
+	: IFileViewTreeItem(bIsShelvedFile ? IChangelistTreeItem::ShelvedFile : IChangelistTreeItem::File)
 	, FileState(InFileState)
 	, MinTimeBetweenUpdate(FTimespan::FromSeconds(5.f))
 	, LastUpdateTime()
@@ -198,21 +216,38 @@ FFileTreeItem::FFileTreeItem(FSourceControlStateRef InFileState, bool bBeautifyP
 	RefreshAssetInformation();
 }
 
+int32 FFileTreeItem::GetIconSortingPriority() const
+{
+	if (!FileState->IsCurrent())        { return 0; } // First if sorted in ascending order.
+	if (FileState->IsUnknown())         { return 1; }
+	if (FileState->IsConflicted())      { return 2; }
+	if (FileState->IsCheckedOutOther()) { return 3; }
+	if (FileState->IsCheckedOut())      { return 4; }
+	if (FileState->IsDeleted())         { return 5; }
+	if (FileState->IsAdded())           { return 6; }
+	else                                { return 7; }
+}
+
+const FString& FFileTreeItem::GetCheckedOutBy() const
+{
+	CheckedOutBy.Reset();
+	FileState->IsCheckedOutOther(&CheckedOutBy);
+	return CheckedOutBy;
+}
+
 FText FFileTreeItem::GetCheckedOutByUser() const
 {
-	FString Users;
-	if (FileState->IsCheckedOutOther(&Users))
-	{
-		return FText::FromString(Users);
-	}
-	return FText::GetEmpty();
+	return FText::FromString(GetCheckedOutBy());
 }
 
 void FFileTreeItem::RefreshAssetInformation()
 {
 	// Initialize display-related members
 	static TArray<FAssetData> NoAssets;
-	RefreshAssetInformationInternal(Assets.IsValid() ? *Assets : NoAssets, FileState->GetFilename(), AssetName, AssetPath, AssetType, PackageName, AssetTypeColor);
+	RefreshAssetInformationInternal(Assets.IsValid() ? *Assets : NoAssets, FileState->GetFilename(), AssetNameStr, AssetPathStr, AssetTypeStr, PackageName, AssetTypeColor);
+	AssetName = FText::FromString(AssetNameStr);
+	AssetPath = FText::FromString(AssetPathStr);
+	AssetType = FText::FromString(AssetTypeStr);
 }
 
 FText FFileTreeItem::GetAssetName() const
@@ -249,7 +284,7 @@ FText FShelvedChangelistTreeItem::GetDisplayText() const
 //////////////////////////////////////////////////////////////////////////
 
 FOfflineFileTreeItem::FOfflineFileTreeItem(const FString& InFilename)
-	: IChangelistTreeItem(IChangelistTreeItem::OfflineFile)
+	: IFileViewTreeItem(IChangelistTreeItem::OfflineFile)
 	, Assets()
 	, Filename(InFilename)
 	, PackageName(FText::FromString(InFilename)) 
@@ -267,7 +302,10 @@ FOfflineFileTreeItem::FOfflineFileTreeItem(const FString& InFilename)
 
 void FOfflineFileTreeItem::RefreshAssetInformation()
 {
-	RefreshAssetInformationInternal(Assets, Filename, AssetName, AssetPath, AssetType, PackageName, AssetTypeColor);
+	RefreshAssetInformationInternal(Assets, Filename, AssetNameStr, AssetPathStr, AssetTypeStr, PackageName, AssetTypeColor);
+	AssetName = FText::FromString(AssetNameStr);
+	AssetPath = FText::FromString(AssetPathStr);
+	AssetType = FText::FromString(AssetTypeStr);
 }
 
 //////////////////////////////////////////////////////////////////////////
diff --git a/Engine/Source/Editor/SourceControlWindows/Private/SSourceControlCommon.h b/Engine/Source/Editor/SourceControlWindows/Private/SSourceControlCommon.h
--- a/Engine/Source/Editor/SourceControlWindows/Private/SSourceControlCommon.h
+++ b/Engine/Source/Editor/SourceControlWindows/Private/SSourceControlCommon.h
@@ -70,6 +70,51 @@ protected:
 };
 
 
+/**
+ * Abstracts the values displayed in the file view that has a set of columns. The API returns values
+ * as string rather than FText to avoid conversion when sorting very large collection (as FText convert
+ * internally to FString for comparison).
+ */
+struct IFileViewTreeItem : public IChangelistTreeItem
+{
+	/**The 'Priority' given to the item icon when sorting ascending (lower will be sorted first). */
+	virtual int32 GetIconSortingPriority() const { return 0; }
+
+	/** The values displayed in the 'Name' column. */
+	virtual const FString& GetName() const { return DefaultStrValue; }
+
+	/** The values displayed in the 'Path' column. */
+	virtual const FString& GetPath() const { return DefaultStrValue; }
+
+	/** The values displayed in the 'Type' column. */
+	virtual const FString& GetType() const { return DefaultStrValue; }
+
+	/** The values displayed in the 'User' column. */
+	virtual const FString& GetCheckedOutBy() const { return DefaultStrValue; }
+
+	/** Set the last modified time timestamp. */
+	void SetLastModifiedDateTime(const FDateTime& Timestamp);
+
+	/** The values displayed in the 'Last Modified' column. */
+	const FDateTime& GetLastModifiedDateTime() const { return LastModifiedDateTime; }
+
+	/** The values displayed in the 'Last Modified' column as text. */
+	FText GetLastModifiedTimestamp() const { return LastModifiedTimestampText; }
+
+protected:
+	IFileViewTreeItem(TreeItemType InType) : IChangelistTreeItem(InType) {}
+
+private:
+	// Use an empty string to return as default for const FString&.
+	static FString DefaultStrValue;
+	static FDateTime DefaultDateTimeValue;
+
+	/** The timestamp of the last modification to the file. */
+	FText LastModifiedTimestampText;
+	FDateTime LastModifiedDateTime;
+};
+
+
 /** Displays a changelist icon/number/description. */
 struct FChangelistTreeItem : public IChangelistTreeItem
 {
@@ -132,10 +177,16 @@ struct FUncontrolledChangelistTreeItem : public IChangelistTreeItem
 
 
 /** Displays a set of files under a changelist or uncontrolled changelist. */
-struct FFileTreeItem : public IChangelistTreeItem
+struct FFileTreeItem : public IFileViewTreeItem
 {
 	explicit FFileTreeItem(FSourceControlStateRef InFileState, bool bBeautifyPaths = true, bool bIsShelvedFile = false);
 
+	virtual int32 GetIconSortingPriority() const override;
+	virtual const FString& GetName() const override { return AssetNameStr; }
+	virtual const FString& GetPath() const override { return AssetPathStr; }
+	virtual const FString& GetType() const override { return AssetTypeStr; }
+	virtual const FString& GetCheckedOutBy() const override;
+
 	/** Updates informations based on AssetData */
 	void RefreshAssetInformation();
 
@@ -154,12 +205,6 @@ struct FFileTreeItem : public IChangelistTreeItem
 	/** Returns the asset type color of the item */
 	FSlateColor GetAssetTypeColor() const { return FSlateColor(AssetTypeColor); }
 
-	/** Returns the last modification time of the file/asset. */
-	FText GetLastModifiedTimestamp() const { return LastModifiedTimestamp; }
-	
-	/** Set the last time the files was saved on disk. */
-	void SetLastModifiedTimestamp(const FText& Timestamp) { LastModifiedTimestamp = Timestamp; }
-
 	/** Returns the user that checked out the file/asset (if any). */
 	FText GetCheckedOutByUser() const;
 
@@ -210,12 +255,15 @@ private:
 
 	/** Cached asset name to display */
 	FText AssetName;
+	FString AssetNameStr;
 
 	/** Cached asset path to display */
 	FText AssetPath;
+	FString AssetPathStr;
 
 	/** Cached asset type to display */
 	FText AssetType;
+	FString AssetTypeStr;
 
 	/** Cached asset type related color to display */
 	FColor AssetTypeColor;
@@ -223,8 +271,8 @@ private:
 	/** Cached package name to display */
 	FText PackageName;
 
-	/** The timestamp of the last modification to the file. */
-	FText LastModifiedTimestamp;
+	/** The other user that has the checked out. */
+	mutable FString CheckedOutBy;
 
 	/** Matching asset(s) to facilitate Locate in content browser */
 	FAssetDataArrayPtr Assets;
@@ -257,30 +305,35 @@ struct FShelvedFileTreeItem : public FFileTreeItem
 };
 
 
-struct FOfflineFileTreeItem : public IChangelistTreeItem
+struct FOfflineFileTreeItem : public IFileViewTreeItem
 {
 	explicit FOfflineFileTreeItem(const FString& InFilename);
 
 	void RefreshAssetInformation();
+
 public:
+	virtual const FString& GetName() const override { return AssetNameStr; }
+	virtual const FString& GetPath() const override { return AssetPathStr; }
+	virtual const FString& GetType() const override { return AssetTypeStr; }
+
 	const FString& GetFilename() const { return Filename; }
 	const FText& GetPackageName() const { return PackageName; }
 	const FText& GetDisplayName() const { return AssetName; }
 	const FText& GetDisplayPath() const { return AssetPath; }
 	const FText& GetDisplayType() const { return AssetType; }
 	const FColor& GetDisplayColor() const { return AssetTypeColor; }
-	FText GetLastModifiedTimestamp() const { return LastModifiedTimestamp; }
-	void SetLastModifiedTimestamp(const FText& Timestamp) { LastModifiedTimestamp = Timestamp; }
 
 private:
 	TArray<FAssetData> Assets;
 	FString Filename;
 	FText PackageName;
 	FText AssetName;
+	FString AssetNameStr;
 	FText AssetPath;
+	FString AssetPathStr;
 	FText AssetType;
+	FString AssetTypeStr;
 	FColor AssetTypeColor;
-	FText LastModifiedTimestamp;
 };
 
 
