Fixed reverting non-packages files from the changelist window that would never revert because they were not packages.

diff --git a/Engine/Source/Editor/SourceControlWindows/Private/SSourceControlChangelists.cpp b/Engine/Source/Editor/SourceControlWindows/Private/SSourceControlChangelists.cpp
--- a/Engine/Source/Editor/SourceControlWindows/Private/SSourceControlChangelists.cpp
+++ b/Engine/Source/Editor/SourceControlWindows/Private/SSourceControlChangelists.cpp
@@ -33,6 +33,7 @@
 #include "AssetToolsModule.h"
 #include "Misc/CString.h"
 #include "Misc/MessageDialog.h"
+#include "Misc/PackageName.h"
 #include "Misc/ScopedSlowTask.h"
 #include "Algo/AnyOf.h"
 #include "HAL/PlatformTime.h"
@@ -1577,9 +1578,23 @@ void SSourceControlChangelistsWidget::OnRevert()
 
 		if (!SelectedControlledFiles.IsEmpty())
 		{
-			SSourceControlCommon::ExecuteChangelistOperationWithSlowTaskWrapper(LOCTEXT("Reverting_Files", "Reverting file(s)..."), [&SelectedControlledFiles]()
+			TArray<FString> PackageFiles;
+			TArray<FString> NonPackageFiles;
+			for (const FString& Filename : SelectedControlledFiles)
 			{
-				if (SourceControlHelpers::RevertAndReloadPackages(SelectedControlledFiles))
+				if (FPackageName::IsPackageFilename(Filename))
+				{
+					PackageFiles.Emplace(Filename);
+				}
+				else
+				{
+					NonPackageFiles.Emplace(Filename);
+				}
+			}
+
+			SSourceControlCommon::ExecuteChangelistOperationWithSlowTaskWrapper(LOCTEXT("Reverting_Files", "Reverting file(s)..."), [&PackageFiles, &NonPackageFiles]()
+			{
+				if ((PackageFiles.Num() == 0 || SourceControlHelpers::RevertAndReloadPackages(PackageFiles)) && (NonPackageFiles.Num() == 0 || SourceControlHelpers::RevertFiles(NonPackageFiles)))
 				{
 					SSourceControlCommon::DisplaySourceControlOperationNotification(LOCTEXT("Revert_Files_Succeeded", "The selected file(s) were reverted."), SNotificationItem::CS_Success);
 				}
