Optimized the SSourceControlChangelist to prevent copying one string per selected item several times just to display the context menu.
  - The function GetSelectedFiles() copied the full pathname of each selected files (can be thoudands) several times when opening the context menu, making it non-responsive.

diff --git a/Engine/Source/Editor/SourceControlWindows/Private/SSourceControlChangelists.cpp b/Engine/Source/Editor/SourceControlWindows/Private/SSourceControlChangelists.cpp
--- a/Engine/Source/Editor/SourceControlWindows/Private/SSourceControlChangelists.cpp
+++ b/Engine/Source/Editor/SourceControlWindows/Private/SSourceControlChangelists.cpp
@@ -733,11 +733,15 @@ void SSourceControlChangelistsWidget::EndRefreshStatus()
 
 void SSourceControlChangelistsWidget::ClearChangelistsTree()
 {
-	if (!ChangelistTreeNodes.IsEmpty() || !UncontrolledChangelistTreeNodes.IsEmpty())
+	if (!ChangelistTreeNodes.IsEmpty())
 	{
 		ChangelistTreeNodes.Reset();
-		UncontrolledChangelistTreeNodes.Reset();
 		ChangelistTreeView->RequestTreeRefresh();
+	}
+
+	if (!UncontrolledChangelistTreeNodes.IsEmpty())
+	{
+		UncontrolledChangelistTreeNodes.Reset();
 		UncontrolledChangelistTreeView->RequestTreeRefresh();
 	}
 
@@ -1179,12 +1183,39 @@ void SSourceControlChangelistsWidget::SetSelectedFiles(const TArray<FString>& Fi
 	}
 }
 
+bool SSourceControlChangelistsWidget::HasFilesSelected() const
+{
+	if (FileTreeView->GetNumItemsSelected() == 0)
+	{
+		return false;
+	}
+	if (ChangelistTreeView->GetNumItemsSelected() == 0)
+	{
+		return false;
+	}
+	switch (ChangelistTreeView->GetSelectedItems()[0]->GetTreeItemType())
+	{
+	case IChangelistTreeItem::Changelist:
+		return true; // If this type of changeslist is selected and the file view as item selected, it means files are selected.
+	case IChangelistTreeItem::UncontrolledChangelist:
+		for (const TSharedPtr<IChangelistTreeItem>& Item : FileTreeView->GetSelectedItems())
+		{
+			if (Item->GetTreeItemType() == IChangelistTreeItem::File)
+			{
+				return true; // Early out.
+			}
+		}
+	default:
+		return false;
+	}
+
+	return false;
+}
+
 TArray<FString> SSourceControlChangelistsWidget::GetSelectedFiles()
 {
-	TArray<FChangelistTreeItemPtr> SelectedItems = FileTreeView->GetSelectedItems();
 	TArray<FString> Files;
-
-	for (const TSharedPtr<IChangelistTreeItem>& Item : SelectedItems)
+	for (const TSharedPtr<IChangelistTreeItem>& Item : FileTreeView->GetSelectedItems())
 	{
 		if (Item->GetTreeItemType() == IChangelistTreeItem::File)
 		{
@@ -1197,9 +1228,7 @@ TArray<FString> SSourceControlChangelistsWidget::GetSelectedFiles()
 
 void SSourceControlChangelistsWidget::GetSelectedFiles(TArray<FString>& OutControlledFiles, TArray<FString>& OutUncontrolledFiles)
 {
-	TArray<FChangelistTreeItemPtr> SelectedItems = FileTreeView->GetSelectedItems();
-
-	for (const TSharedPtr<IChangelistTreeItem>& Item : SelectedItems)
+	for (const TSharedPtr<IChangelistTreeItem>& Item : FileTreeView->GetSelectedItems())
 	{
 		if (Item->GetTreeItemType() == IChangelistTreeItem::File)
 		{
@@ -1256,6 +1285,13 @@ void SSourceControlChangelistsWidget::GetSelectedFileStates(TArray<FSourceContro
 	}
 }
 
+bool SSourceControlChangelistsWidget::HasShelvedFilesSelected() const
+{
+	return FileTreeView->GetNumItemsSelected() > 0 &&
+	       ChangelistTreeView->GetNumItemsSelected() > 0 &&
+	       ChangelistTreeView->GetSelectedItems()[0]->GetTreeItemType() == IChangelistTreeItem::ShelvedChangelist;
+}
+
 TArray<FString> SSourceControlChangelistsWidget::GetSelectedShelvedFiles()
 {
 	TArray<FString> ShelvedFiles;
@@ -1526,7 +1562,7 @@ void SSourceControlChangelistsWidget::OnRevertUnchanged()
 
 bool SSourceControlChangelistsWidget::CanRevertUnchanged()
 {
-	return GetSelectedFiles().Num() > 0 || (GetCurrentChangelistState() && GetCurrentChangelistState()->GetFilesStates().Num() > 0);
+	return HasFilesSelected() || (GetCurrentChangelistState() && GetCurrentChangelistState()->GetFilesStates().Num() > 0);
 }
 
 void SSourceControlChangelistsWidget::OnRevert()
@@ -1630,7 +1666,7 @@ bool SSourceControlChangelistsWidget::CanRevert()
 	FSourceControlChangelistStatePtr CurrentChangelistState = GetCurrentChangelistState();
 	FUncontrolledChangelistStatePtr CurrentUncontrolledChangelistState = GetCurrentUncontrolledChangelistState();
 
-	return GetSelectedFiles().Num() > 0
+	return HasFilesSelected()
 		|| (CurrentChangelistState.IsValid() && CurrentChangelistState->GetFilesStates().Num() > 0)
 		|| (CurrentUncontrolledChangelistState.IsValid() && CurrentUncontrolledChangelistState->GetFilesStates().Num() > 0);
 }
@@ -2296,7 +2332,7 @@ void SSourceControlChangelistsWidget::OnDiffAgainstDepot()
 
 bool SSourceControlChangelistsWidget::CanDiffAgainstDepot()
 {
-	return GetSelectedFiles().Num() == 1;
+	return FileTreeView->GetNumItemsSelected() == 1 && HasFilesSelected();
 }
 
 void SSourceControlChangelistsWidget::OnDiffAgainstWorkspace()
@@ -2310,7 +2346,7 @@ void SSourceControlChangelistsWidget::OnDiffAgainstWorkspace()
 
 bool SSourceControlChangelistsWidget::CanDiffAgainstWorkspace()
 {
-	return GetSelectedShelvedFiles().Num() == 1;
+	return FileTreeView->GetNumItemsSelected() == 1 && HasShelvedFilesSelected();
 }
 
 TSharedPtr<SWidget> SSourceControlChangelistsWidget::OnOpenContextMenu()
@@ -2330,8 +2366,8 @@ TSharedPtr<SWidget> SSourceControlChangelistsWidget::OnOpenContextMenu()
 	bool bHasSelectedChangelist = SelectedChangelistNodes.Num() > 0 &&  SelectedChangelistNodes[0]->GetTreeItemType() == IChangelistTreeItem::Changelist;
 	bool bHasSelectedShelvedChangelistNode = SelectedChangelistNodes.Num() > 0 &&  SelectedChangelistNodes[0]->GetTreeItemType() == IChangelistTreeItem::ShelvedChangelist;
 	bool bHasSelectedUncontrolledChangelist = SelectedUncontrolledChangelistNodes.Num() > 0 &&  SelectedUncontrolledChangelistNodes[0]->GetTreeItemType() == IChangelistTreeItem::UncontrolledChangelist;
-	bool bHasSelectedFiles = (GetSelectedFiles().Num() > 0);
-	bool bHasSelectedShelvedFiles = (GetSelectedShelvedFiles().Num() > 0);
+	bool bHasSelectedFiles = HasFilesSelected();
+	bool bHasSelectedShelvedFiles = HasShelvedFilesSelected();
 	bool bHasEmptySelection = (!bHasSelectedChangelist && !bHasSelectedFiles && !bHasSelectedShelvedFiles);
 
 	// Build up the menu for a selection
@@ -2719,7 +2755,7 @@ FReply SSourceControlChangelistsWidget::OnFilesDragged(const FGeometry& InGeomet
 	{
 		TSharedRef<FSCCFileDragDropOp> Operation = MakeShared<FSCCFileDragDropOp>();
 
-		for (FChangelistTreeItemPtr InTreeItem : FileTreeView->GetSelectedItems())
+		for (const FChangelistTreeItemPtr& InTreeItem : FileTreeView->GetSelectedItems())
 		{
 			if (InTreeItem->GetTreeItemType() == IChangelistTreeItem::File)
 			{
@@ -2804,6 +2840,8 @@ void SSourceControlChangelistsWidget::OnItemDoubleClicked(TSharedPtr<IChangelist
 
 void SSourceControlChangelistsWidget::OnChangelistSelectionChanged(TSharedPtr<IChangelistTreeItem> SelectedChangelist, ESelectInfo::Type SelectionType)
 {
+	TRACE_CPUPROFILER_EVENT_SCOPE(SSourceControlChangelistsWidget::OnChangelistSelectionChanged);
+
 	FileTreeNodes.Reset();
 
 	// Add the children of the parent item to the file tree node.
diff --git a/Engine/Source/Editor/SourceControlWindows/Private/SSourceControlChangelists.h b/Engine/Source/Editor/SourceControlWindows/Private/SSourceControlChangelists.h
--- a/Engine/Source/Editor/SourceControlWindows/Private/SSourceControlChangelists.h
+++ b/Engine/Source/Editor/SourceControlWindows/Private/SSourceControlChangelists.h
@@ -50,6 +50,7 @@ private:
 	};
 
 private:
+	TSharedRef<SWidget> MakeToolBar();
 	TSharedRef<SChangelistTree> CreateChangelistTreeView(TArray<TSharedPtr<IChangelistTreeItem>>& ItemSources);
 	TSharedRef<STreeView<FChangelistTreeItemPtr>> CreateChangelistFilesView();
 
@@ -88,6 +89,11 @@ private:
 
 	/** Returns list of currently selected files */
 	TArray<FString> GetSelectedFiles();
+	bool HasFilesSelected() const;
+
+	/** Returns list of currently selected shelved files */
+	TArray<FString> GetSelectedShelvedFiles();
+	bool HasShelvedFilesSelected() const;
 
 	/**
 	 * Splits selected files between Controlled and Uncontrolled files.
@@ -103,9 +109,6 @@ private:
 	 */
 	void GetSelectedFileStates(TArray<FSourceControlStateRef>& OutControlledFileStates, TArray<FSourceControlStateRef>& OutUncontrolledFileStates);
 
-	/** Returns list of currently selected shelved files */
-	TArray<FString> GetSelectedShelvedFiles();
-
 	/** Intercept Enter and Delete key presses to Submit or Delete the selected changelist (if conditions are met) */
 	virtual FReply OnKeyDown(const FGeometry& MyGeometry, const FKeyEvent& InKeyEvent) override;
 
@@ -171,8 +174,6 @@ private:
 	void SaveExpandedAndSelectionStates(FExpandedAndSelectionStates& OutStates);
 	void RestoreExpandedAndSelectionStates(const FExpandedAndSelectionStates& InStates);
 
-	TSharedRef<SWidget> MakeToolBar();
-	
 	/** Executes an operation to updates the changelist description of the provided changelist with a new description. */
 	void EditChangelistDescription(const FText& InNewChangelistDescription, const FSourceControlChangelistStatePtr& InChangelistState);
 
diff --git a/Engine/Source/Editor/SourceControlWindows/Private/SSourceControlCommon.cpp b/Engine/Source/Editor/SourceControlWindows/Private/SSourceControlCommon.cpp
--- a/Engine/Source/Editor/SourceControlWindows/Private/SSourceControlCommon.cpp
+++ b/Engine/Source/Editor/SourceControlWindows/Private/SSourceControlCommon.cpp
@@ -178,14 +178,17 @@ FDateTime IFileViewTreeItem::DefaultDateTimeValue; // Default is FDateTime::MinV
 
 void IFileViewTreeItem::SetLastModifiedDateTime(const FDateTime& Timestamp)
 {
-	LastModifiedDateTime = Timestamp;
-	if (Timestamp != FDateTime::MinValue())
+	if (Timestamp != LastModifiedDateTime) // Pay the text conversion only if needed.
 	{
-		LastModifiedTimestampText = FText::AsDateTime(Timestamp, EDateTimeStyle::Short);
-	}
-	else
-	{
-		LastModifiedTimestampText = FText::GetEmpty();
+		LastModifiedDateTime = Timestamp;
+		if (Timestamp != FDateTime::MinValue())
+		{
+			LastModifiedTimestampText = FText::AsDateTime(Timestamp, EDateTimeStyle::Short);
+		}
+		else
+		{
+			LastModifiedTimestampText = FText::GetEmpty();
+		}
 	}
 }
 
