Optimized source control changelist to handle a reasonably large number of files smootly ([REDACTED_PHONE]).
  - Handling simple source control updates nofications, for example when user scroll the outliner which query source control status of the scrolled files, usually took 450ms with 15000 files to refresh in the changelist window now takes about 3ms.
  - Cached UI item to reuse them when user change the filter.
  - Detect and update only the item that changed.
  - Preserved the item sort order from the previous update when possible.
  - Lazily read the file modification timestamp only when the item is fist displayed.
  - Updated files modification in background (triggred when when a package is saved).
  - Only monitor source control file status for files that are visible (unless the view is sorted by status)
  - Splitted Uncontrolled Changelist and Source Controlled Changelist update to update then independently and just  update the one that was notified as updated.
  - Converted the file view from a STreeView to a SListView.

Updated changelist UI to display (X out of Y files) when a filter is active and remove files from the view.

diff --git a/Engine/Source/Editor/SourceControlWindows/Private/SSourceControlChangelistRows.cpp b/Engine/Source/Editor/SourceControlWindows/Private/SSourceControlChangelistRows.cpp
--- a/Engine/Source/Editor/SourceControlWindows/Private/SSourceControlChangelistRows.cpp
+++ b/Engine/Source/Editor/SourceControlWindows/Private/SSourceControlChangelistRows.cpp
@@ -36,6 +36,12 @@ FName SourceControlFileViewColumn::CheckedOutByUser::Id() { return TEXT("Checked
 FText SourceControlFileViewColumn::CheckedOutByUser::GetDisplayText() { return LOCTEXT("CheckedOutByUser_Column", "User"); }
 FText SourceControlFileViewColumn::CheckedOutByUser::GetToolTipText() { return LOCTEXT("CheckedOutByUser_Column_Tooltip", "Displays the other user(s) that checked out the file/asset, if any"); }
 
+FText FormatChangelistFileCountText(int32 DisplayedCount, int32 TotalCount)
+{
+	return DisplayedCount == TotalCount ?
+		FText::Format(INVTEXT("({0})"), TotalCount) :
+		FText::Format(LOCTEXT("FilterNum", "({0} out of {1})"), DisplayedCount, TotalCount);
+}
 
 void SChangelistTableRow::Construct(const FArguments& InArgs, const TSharedRef<STableViewBase>& InOwner)
 {
@@ -73,7 +79,11 @@ void SChangelistTableRow::Construct(const FArguments& InArgs, const TSharedRef<S
 			.AutoWidth()
 			[
 				SNew(STextBlock)
-				.Text(FText::Format(INVTEXT("({0})"), TreeItem->GetFileCount()))
+				.Text_Lambda([this]()
+				{
+					// Check if the 'Shelved Files' node is currently linked to the  tree view. (not filtered out).
+					return FormatChangelistFileCountText(TreeItem->ShelvedChangelistItem->GetParent() ? TreeItem->GetChildren().Num() - 1 : TreeItem->GetChildren().Num(), TreeItem->GetFileCount());
+				})
 			]
 			+SHorizontalBox::Slot() // Description.
 			.Padding(2, 0, 0, 0)
@@ -200,7 +210,10 @@ void SUncontrolledChangelistTableRow::Construct(const FArguments& InArgs, const
 			.AutoWidth()
 			[
 				SNew(STextBlock)
-				.Text(FText::Format(INVTEXT("({0})"), TreeItem->GetFileCount()))
+				.Text_Lambda([this]()
+				{
+					return FormatChangelistFileCountText(TreeItem->GetChildren().Num(), TreeItem->GetFileCount());
+				})
 			]
 		], InOwner);
 }
@@ -375,11 +388,22 @@ void SShelvedFilesTableRow::Construct(const FArguments& InArgs, const TSharedRef
 				+SHorizontalBox::Slot()
 				.Padding(2.0f, 1.0f)
 				.VAlign(VAlign_Center)
+				.AutoWidth()
 				[
 					SNew(STextBlock)
-					.Text(TreeItem->GetDisplayText())
+					.Text_Lambda([this](){ return TreeItem->GetDisplayText(); })
 					.HighlightText(InArgs._HighlightText)
 				]
+				+SHorizontalBox::Slot() // Shelved file count.
+				.Padding(4, 0, 4, 0)
+				.AutoWidth()
+				[
+					SNew(STextBlock)
+					.Text_Lambda([this]()
+					{
+						return FormatChangelistFileCountText(TreeItem->GetChildren().Num(), static_cast<const FChangelistTreeItem*>(TreeItem->GetParent().Get())->GetShelvedFileCount());
+					})
+				]
 		],
 		InOwnerTableView);
 }
diff --git a/Engine/Source/Editor/SourceControlWindows/Private/SSourceControlChangelists.cpp b/Engine/Source/Editor/SourceControlWindows/Private/SSourceControlChangelists.cpp
--- a/Engine/Source/Editor/SourceControlWindows/Private/SSourceControlChangelists.cpp
+++ b/Engine/Source/Editor/SourceControlWindows/Private/SSourceControlChangelists.cpp
@@ -40,6 +40,7 @@
 #include "HAL/FileManager.h"
 #include "ToolMenus.h"
 #include "UObject/UObjectGlobals.h"
+#include "UObject/Package.h"
 #include "SSourceControlSubmit.h"
 #include "Framework/Application/SlateApplication.h"
 #include "Framework/Notifications/NotificationManager.h"
@@ -106,38 +107,6 @@ FText UpdateChangelistDescriptionToSubmitIfNeeded(const bool bInValidationResult
 	return InChangelistDescription;
 }
 
-void UpdateLastModifiedTimestamp(IChangelistTreeItem& Item)
-{
-	switch (Item.GetTreeItemType())
-	{
-		// NOTE: We may not want to display the last saved timestamp for shelved files as this could lead users to think the file
-		//       was shelved at that time. Consider remove this case in the future if this is too confusing.
-		case IChangelistTreeItem::File:
-		case IChangelistTreeItem::ShelvedFile:
-		{
-			FFileTreeItem& FileItem = static_cast<FFileTreeItem&>(Item);
-			FileItem.SetLastModifiedDateTime(IFileManager::Get().GetTimeStamp(*FileItem.GetFileName().ToString()));
-			return;
-		}
-		break;
-
-		case IChangelistTreeItem::OfflineFile:
-		{
-			FOfflineFileTreeItem& OfflineItem = static_cast<FOfflineFileTreeItem&>(Item);
-			OfflineItem.SetLastModifiedDateTime(IFileManager::Get().GetTimeStamp(*OfflineItem.GetFilename()));
-			return;
-		}
-
-		case IChangelistTreeItem::Changelist:
-		case IChangelistTreeItem::UncontrolledChangelist:
-		case IChangelistTreeItem::ShelvedChangelist:
-			return; // We don't update timestamp for those nodes.
-
-		default:
-			checkNoEntry();
-	}
-}
-
 } // Anonymous namespace
 
 
@@ -234,7 +203,7 @@ public:
 		/** Text displayed on the expandable area */
 		SLATE_ATTRIBUTE(FText, HeaderText)
 		/** The tree element displayed as body. */
-		SLATE_ARGUMENT(TSharedPtr<SChangelistTree>, ChangelistView)
+		SLATE_ARGUMENT(TSharedPtr<STreeView<FChangelistTreeItemPtr>>, ChangelistView)
 		/** Event fired when the associated search box is made visible */
 		SLATE_EVENT(FOnSearchBoxExpanded, OnSearchBoxExpanded)
 		/** Event fired when the 'plus' button is clicked. */
@@ -326,6 +295,7 @@ public:
 		];
 	}
 
+	void SetExpanded(bool bExpanded) { ExpandableArea->SetExpanded(bExpanded); }
 	bool IsExpanded() const { return ExpandableArea->IsExpanded(); }
 	FText GetSearchedText() const { return SearchBox->GetText(); }
 	TSharedPtr<SSearchBox> GetSearchBox() { return SearchBox; }
@@ -336,6 +306,22 @@ private:
 };
 
 
+void SSourceControlChangelistsWidget::FAsyncTimestampUpdater::DoWork()
+{
+	for (FString& Pathname : RequestedFileTimestamps)
+	{
+		if (!bAbandon)
+		{
+			if (FPaths::IsRelative(Pathname))
+			{
+				Pathname = FPaths::ConvertRelativePathToFull(Pathname);
+			}
+			QueriedFileTimestamps.Emplace(MoveTemp(Pathname), IFileManager::Get().GetTimeStamp(*Pathname));
+		}
+	}
+}
+
+
 void SSourceControlChangelistsWidget::Construct(const FArguments& InArgs)
 {
 	// Register delegates
@@ -344,14 +330,17 @@ void SSourceControlChangelistsWidget::Construct(const FArguments& InArgs)
 
 	SCCModule.RegisterProviderChanged(FSourceControlProviderChanged::FDelegate::CreateSP(this, &SSourceControlChangelistsWidget::OnSourceControlProviderChanged));
 	SourceControlStateChangedDelegateHandle = SCCModule.GetProvider().RegisterSourceControlStateChanged_Handle(FSourceControlStateChanged::FDelegate::CreateSP(this, &SSourceControlChangelistsWidget::OnSourceControlStateChanged));
-	UncontrolledChangelistModule.OnUncontrolledChangelistModuleChanged.AddSP(this, &SSourceControlChangelistsWidget::OnSourceControlStateChanged);
+	UncontrolledChangelistModule.OnUncontrolledChangelistModuleChanged.AddSP(this, &SSourceControlChangelistsWidget::OnUncontrolledChangelistStateChanged);
+
+	// Monitor when packages are saved to update the timestamps.
+	UPackage::PackageSavedWithContextEvent.AddSP(this, &SSourceControlChangelistsWidget::OnPackageSaved);
 
 	// No default sorting, sorting make the view slower, just pay the cost if user wants it.
 	PrimarySortedColumn = FName();
 
 	ChangelistTreeView = CreateChangelistTreeView(ChangelistTreeNodes);
 	UncontrolledChangelistTreeView = CreateChangelistTreeView(UncontrolledChangelistTreeNodes);
-	FileTreeView = CreateChangelistFilesView();
+	FileListView = CreateChangelistFilesView();
 
 	ChangelistExpandableArea = SNew(SExpandableChangelistArea)
 		.HeaderText_Lambda([this]() { return FText::Format(LOCTEXT("SourceControl_ChangeLists", "Changelists ({0})"), ChangelistTreeNodes.Num()); })
@@ -370,13 +359,13 @@ void SSourceControlChangelistsWidget::Construct(const FArguments& InArgs)
 		.OnSearchTextChanged(this, &SSourceControlChangelistsWidget::OnUncontrolledChangelistSearchTextChanged);
 
 	ChangelistTextFilter = MakeShared<TTextFilter<const IChangelistTreeItem&>>(TTextFilter<const IChangelistTreeItem&>::FItemToStringArray::CreateSP(this, &SSourceControlChangelistsWidget::PopulateItemSearchStrings));
-	ChangelistTextFilter->OnChanged().AddSP(this, &SSourceControlChangelistsWidget::OnRefresh);
+	ChangelistTextFilter->OnChanged().AddSP(this, &SSourceControlChangelistsWidget::OnRefreshUI, ERefreshFlags::SourceControlChangelists);
 
 	UncontrolledChangelistTextFilter = MakeShared<TTextFilter<const IChangelistTreeItem&>>(TTextFilter<const IChangelistTreeItem&>::FItemToStringArray::CreateSP(this, &SSourceControlChangelistsWidget::PopulateItemSearchStrings));
-	UncontrolledChangelistTextFilter->OnChanged().AddSP(this, &SSourceControlChangelistsWidget::OnRefresh);
+	UncontrolledChangelistTextFilter->OnChanged().AddSP(this, &SSourceControlChangelistsWidget::OnRefreshUI, ERefreshFlags::UncontrolledChangelists);
 
 	FileTextFilter = MakeShared<TTextFilter<const IChangelistTreeItem&>>(TTextFilter<const IChangelistTreeItem&>::FItemToStringArray::CreateSP(this, &SSourceControlChangelistsWidget::PopulateItemSearchStrings));
-	FileTextFilter->OnChanged().AddSP(this, &SSourceControlChangelistsWidget::OnRefresh);
+	FileTextFilter->OnChanged().AddSP(this, &SSourceControlChangelistsWidget::OnRefreshUI, ERefreshFlags::All);
 
 	// Min/Max prevents making the Changelist Area too small and consequently prevent the file search box to go over the 'refresh' button.
 	const float MinChangelistAreaRatio = 0.2f;
@@ -502,7 +491,7 @@ void SSourceControlChangelistsWidget::Construct(const FArguments& InArgs)
 				.Value_Lambda([this, MaxFileAreaRation]() { return FMath::Min(MaxFileAreaRation, FileAreaSize); })
 				.OnSlotResized_Lambda([this](float Size) { FileAreaSize = Size; })
 				[
-					FileTreeView.ToSharedRef()
+					FileListView.ToSharedRef()
 				]
 			]
 		]
@@ -536,6 +525,18 @@ void SSourceControlChangelistsWidget::Construct(const FArguments& InArgs)
 	bShouldRefresh = true;
 }
 
+SSourceControlChangelistsWidget::~SSourceControlChangelistsWidget()
+{
+	ISourceControlModule::Get().GetSourceControlFileStatusMonitor().StopMonitoringFiles(reinterpret_cast<uintptr_t>(this));
+
+	if (TimestampUpdateTask)
+	{
+		TimestampUpdateTask->TryAbandonTask();
+		TimestampUpdateTask->EnsureCompletion(/*bDoWorkOnThisThreadIfNotStarted*/false);
+		TimestampUpdateTask.Reset();
+	}
+}
+
 TSharedRef<SWidget> SSourceControlChangelistsWidget::MakeToolBar()
 {
 	FSlimHorizontalToolBarBuilder ToolBarBuilder(nullptr, FMultiBoxCustomization::None);
@@ -594,6 +595,11 @@ void SSourceControlChangelistsWidget::EditChangelistDescription(const FText& InN
 
 void SSourceControlChangelistsWidget::Tick(const FGeometry& AllottedGeometry, const double InCurrentTime, const float InDeltaTime)
 {
+	if (!ISourceControlModule::Get().IsEnabled() && !FUncontrolledChangelistsModule::Get().IsEnabled())
+	{
+		return;
+	}
+
 	// Detect transitions of the source control being available/unavailable. Ex: When the user changes the source control in UI, the provider gets selected,
 	// but it is not connected/available until the user accepts the settings. The source control doesn't have callback for availability and we want to refresh everything
 	// once it gets available.
@@ -605,23 +611,118 @@ void SSourceControlChangelistsWidget::Tick(const FGeometry& AllottedGeometry, co
 
 	if (bShouldRefresh)
 	{
-		if (ISourceControlModule::Get().IsEnabled() || FUncontrolledChangelistsModule::Get().IsEnabled())
+		if (bInitialRefreshDone)
 		{
-			RequestChangelistsRefresh();
-		}
-		else
-		{
-			// No provider available, clear changelist tree
-			ClearChangelistsTree();
+			// Ensure the UI is built from the current cached states once because all UI updates are hooked on 'state change' callbacks and those will only
+			// be called when the state changed.
+			OnRefreshUI(ERefreshFlags::All);
+			bInitialRefreshDone = true;
 		}
 
+		RequestChangelistsRefresh();
 		bShouldRefresh = false;
 	}
-	
+
 	if (bIsRefreshing)
 	{
 		TickRefreshStatus(InDeltaTime);
 	}
+
+	if (bUpdateMonitoredFileStatusList)
+	{
+		if (FileListNodes.IsEmpty())
+		{
+			// Nothing to monitor.
+			ISourceControlModule::Get().GetSourceControlFileStatusMonitor().StopMonitoringFiles(reinterpret_cast<uintptr_t>(this));
+			bUpdateMonitoredFileStatusList = false;
+		}
+		else if (IsFileViewSortedByFileStatusIcon())
+		{
+			// If a changelist is selected, monitor all files status to honor the sort order on that column.
+			if (ChangelistTreeView->GetNumItemsSelected() > 0)
+			{
+				RequestFileStatusRefresh(*ChangelistTreeView->GetSelectedItems()[0]);
+			}
+			else if (UncontrolledChangelistTreeView->GetNumItemsSelected() > 0)
+			{
+				RequestFileStatusRefresh(*UncontrolledChangelistTreeView->GetSelectedItems()[0]);
+			}
+			bUpdateMonitoredFileStatusList = false;
+		}
+		else // Just monitor the files that are visible in the view.
+		{
+			// NOTE: In the tick where the view is updated/refreshed, the item count is usually zero, next tick it will have the correct value, so wait for it.
+			int32 ItemCount = FileListView->GetNumGeneratedChildren();
+			if (ItemCount > 0)
+			{
+				// Minimize the list of file status monitored to only include those in view. The ones outside of the visible area don't need
+				// to be monitored unless the user sorts.
+				int32 ItemStart = static_cast<int32>(FileListView->GetScrollOffset()); // This give the offset in item count.
+				TSet<FString> MonitoredFiles;
+				for (int32 Index = ItemStart; Index < ItemCount; ++Index)
+				{
+					MonitoredFiles.Emplace(StaticCastSharedPtr<IFileViewTreeItem>(FileListNodes[Index])->GetFullPathname());
+				}
+				RequestFileStatusRefresh(MonitoredFiles);
+				bUpdateMonitoredFileStatusList = false;
+			}
+		}
+	}
+
+	// If requested timestamps were gathered in background.
+	if (TimestampUpdateTask && TimestampUpdateTask->IsDone())
+	{
+		bool bUpdatedItemInView = false;
+
+		// Start updating the offline files (if any) as we have a direct mapping between keys.
+		if (OfflineFileItemCache.Num() > 0)
+		{
+			for (auto It = TimestampUpdateTask->GetTask().QueriedFileTimestamps.CreateIterator(); It; ++It)
+			{
+				if (TSharedPtr<IFileViewTreeItem>* FileViewItem = OfflineFileItemCache.Find(It->Key))
+				{
+					(*FileViewItem)->SetLastModifiedDateTime(It->Value);
+					bUpdatedItemInView |= (*FileViewItem)->DisplayedUpdateNum == UpdateRequestNum;
+					It.RemoveCurrent();
+				}
+			}
+		}
+
+		// Update other files. There is no direct lookup key File -> FileItem, do a linear scan of the cache.
+		int32 UpdateRemaining = TimestampUpdateTask->GetTask().QueriedFileTimestamps.Num();
+		auto UpdateCachedItems = [this, &UpdateRemaining, &bUpdatedItemInView](TMap<TSharedPtr<void>, TSharedPtr<IChangelistTreeItem>>& Cache)
+		{
+			for (auto It = Cache.CreateIterator(); It && UpdateRemaining > 0; ++It)
+			{
+				if (It->Value->GetTreeItemType() == IChangelistTreeItem::File || It->Value->GetTreeItemType() == IChangelistTreeItem::ShelvedFile)
+				{
+					IFileViewTreeItem* FileViewItem = static_cast<IFileViewTreeItem*>(It->Value.Get());
+					if (const FDateTime* Timestamp = TimestampUpdateTask->GetTask().QueriedFileTimestamps.Find(FileViewItem->GetFullPathname()))
+					{
+						FileViewItem->SetLastModifiedDateTime(*Timestamp);
+						bUpdatedItemInView |= FileViewItem->DisplayedUpdateNum == UpdateRequestNum;
+						--UpdateRemaining;
+					}
+				}
+			}
+		};
+		UpdateCachedItems(SourceControlItemCache);
+		UpdateCachedItems(UncontrolledChangelistItemCache);
+		TimestampUpdateTask.Reset();
+
+		if (bUpdatedItemInView && IsFileViewSortedByLastModifiedTimestamp())
+		{
+			SortFileView();
+			FileListView->RequestListRefresh();
+		}
+	}
+	if (!TimestampUpdateTask && OutdatedTimestampFiles.Num()) // Should be sufficient to run only one async task at the time.
+	{
+		// Query the timestamps asynchronously.
+		TimestampUpdateTask = MakeUnique<FAsyncTask<FAsyncTimestampUpdater>>();
+		TimestampUpdateTask->GetTask().RequestedFileTimestamps = MoveTemp(OutdatedTimestampFiles);
+		TimestampUpdateTask->StartBackgroundTask();
+	}
 }
 
 void SSourceControlChangelistsWidget::RequestChangelistsRefresh()
@@ -745,166 +846,430 @@ void SSourceControlChangelistsWidget::ClearChangelistsTree()
 		UncontrolledChangelistTreeView->RequestTreeRefresh();
 	}
 
-	if (!FileTreeNodes.IsEmpty())
+	if (!FileListNodes.IsEmpty())
 	{
-		FileTreeNodes.Reset();
-		FileTreeView->RequestTreeRefresh();
+		FileListNodes.Reset();
+		FileListView->RequestListRefresh();
 	}
+
+	SourceControlItemCache.Reset();
+	UncontrolledChangelistItemCache.Reset();
+	OfflineFileItemCache.Reset();
 }
 
-void SSourceControlChangelistsWidget::OnRefresh()
+void SSourceControlChangelistsWidget::OnRefreshSourceControlWidgets(int64 CurrUpdateNum, const TFunction<void(TSharedPtr<IFileViewTreeItem>&)>& AddItemToFileView)
 {
-	if (!AreChangelistsEnabled())
-	{
-		ClearChangelistsTree();
-		return;
-	}
+	TRACE_CPUPROFILER_EVENT_SCOPE(SSourceControlChangelistsWidget::OnRefreshSourceControlWidgets);
 
-	TRACE_CPUPROFILER_EVENT_SCOPE(SSourceControlChangelistsWidget::Refresh);
+	const bool bBeautifyPaths = true;
+	const int64 PrevUpdateNum = CurrUpdateNum - 1;
+	int32 DisplayedItemCount = 0;
+	int32 VisitedItemCount = 0;
 
-	// Views will be teared down and rebuilt from scratch, save the items that are expanded and/or selected to be able to restore those states after the rebuild.
-	FExpandedAndSelectionStates ExpandedAndSelectedStates;
-	SaveExpandedAndSelectionStates(ExpandedAndSelectedStates);
+	TSharedPtr<IChangelistTreeItem> SelectedChangelistItem = ChangelistTreeView->GetNumItemsSelected() > 0 ? ChangelistTreeView->GetSelectedItems()[0] : TSharedPtr<IChangelistTreeItem>();
 
-	// Query the source control
-	ISourceControlProvider& SourceControlProvider = ISourceControlModule::Get().GetProvider();
-	FUncontrolledChangelistsModule& UncontrolledChangelistModule = FUncontrolledChangelistsModule::Get();
-	TArray<FSourceControlChangelistRef> Changelists = SourceControlProvider.GetChangelists(EStateCacheUsage::Use);
-	TArray<TSharedRef<FUncontrolledChangelistState>> UncontrolledChangelistStates = UncontrolledChangelistModule.GetChangelistStates();
+	auto HasSourceControlChangelistSelected = [&SelectedChangelistItem]()
+	{
+		return SelectedChangelistItem && 
+			(SelectedChangelistItem->GetTreeItemType() == IChangelistTreeItem::Changelist || SelectedChangelistItem->GetTreeItemType() == IChangelistTreeItem::ShelvedChangelist);
+	};
 
-	TArray<FSourceControlChangelistStateRef> ChangelistsStates;
-	SourceControlProvider.GetState(Changelists, ChangelistsStates, EStateCacheUsage::Use);
+	auto OnSourceControlCachedItemVisited = [CurrUpdateNum, &VisitedItemCount](IChangelistTreeItem* Item)
+	{
+		checkf(Item->VisitedUpdateNum < CurrUpdateNum, TEXT("The same source control item was visited twice. It is likely present in more than one changelist"));
+		Item->VisitedUpdateNum = CurrUpdateNum;
+		++VisitedItemCount;
+	};
 
-	// Count number of steps for slow task...
-	int32 ElementsToProcess = ChangelistsStates.Num();
-	ElementsToProcess += UncontrolledChangelistStates.Num();
+	auto OnNewSourceControlItemCreated = [this, &OnSourceControlCachedItemVisited](TSharedPtr<void> Key, TSharedPtr<IChangelistTreeItem> Item)
+	{
+		OnSourceControlCachedItemVisited(Item.Get());
+		SourceControlItemCache.Emplace(MoveTemp(Key), MoveTemp(Item));
+	};
 
-	for (const TSharedRef<ISourceControlChangelistState>& ChangelistState : ChangelistsStates)
+	auto AddItemToFileViewWrapper = [&AddItemToFileView, &DisplayedItemCount](TSharedPtr<IFileViewTreeItem>& Item)
 	{
-		ElementsToProcess += ChangelistState->GetFilesStates().Num();
-		ElementsToProcess += ChangelistState->GetShelvedFilesStates().Num();
-	}
+		++DisplayedItemCount;
+		AddItemToFileView(Item);
+	};
 
-	for (const TSharedRef<FUncontrolledChangelistState>& UncontrolledChangelistState : UncontrolledChangelistStates)
+	auto RemoveDiscardedSourceControlItems = [this, VisitedItemCount, DisplayedItemCount, CurrUpdateNum, HasSourceControlChangelistSelected]()
 	{
-		ElementsToProcess += UncontrolledChangelistState->GetFileCount();
-	}
+		TRACE_CPUPROFILER_EVENT_SCOPE(SSourceControlChangelistsWidget::RemoveDiscardedSourceControlItems);
+
+		check(SourceControlItemCache.Num() >= VisitedItemCount); // Cannot visit items that are not cached.
 
-	FScopedSlowTask SlowTask(ElementsToProcess, LOCTEXT("SourceControl_RebuildTree", "Refreshing Tree Items"));
-	SlowTask.MakeDialogDelayed(1.5f, /*bShowCancelButton=*/true);
+		// Remove the items used to be visited but weren't visited during this iteration.
+		int32 DiscardCount = SourceControlItemCache.Num() - VisitedItemCount;
+		for (auto It = SourceControlItemCache.CreateIterator(); It && DiscardCount > 0 ; ++It)
+		{
+			if (It->Value->VisitedUpdateNum < CurrUpdateNum)
+			{
+				It.RemoveCurrent();
+				--DiscardCount;
+			}
+		}
+
+		// If the source control feeds the file view.
+		if (HasSourceControlChangelistSelected())
+		{
+			// Unless the FileTreeNode array was reset because the selected changelist changed, it still contains the union of items to display for this iteration and the item that were displayed last
+			// iteration, so we need to remove item that aren't displayed anymore. (The FileTreeNode is not reset between updates to preserve the sort order as much as possible).
+			if (FileListNodes.Num() > DisplayedItemCount)
+			{
+				// Remove items that were not displayed anymore.
+				FileListNodes.RemoveAll([CurrUpdateNum](const TSharedPtr<IChangelistTreeItem>& Candidate) { return Candidate->DisplayedUpdateNum < CurrUpdateNum; });
+			}
+		}
+	};
+
+	// Query the source control
+	ISourceControlProvider& SourceControlProvider = ISourceControlModule::Get().GetProvider();
+	TArray<FSourceControlChangelistRef> Changelists = SourceControlProvider.GetChangelists(EStateCacheUsage::Use);
+	TArray<FSourceControlChangelistStateRef> ChangelistsStates;
+	SourceControlProvider.GetState(Changelists, ChangelistsStates, EStateCacheUsage::Use);
 
-	// Rebuild the tree data models
-	bool bBeautifyPaths = true;
 	ChangelistTreeNodes.Reset(ChangelistsStates.Num());
-	UncontrolledChangelistTreeNodes.Reset(UncontrolledChangelistStates.Num());
-	FileTreeNodes.Reset();
 
+	for (const TSharedRef<ISourceControlChangelistState>& ChangelistState : ChangelistsStates)
 	{
-		TRACE_CPUPROFILER_EVENT_SCOPE(SSourceControlChangelistsWidget::RefreshControlledChangelists);
+		// Create or reuse the 'Changelists' and 'Shelved Files' node.
+		TSharedPtr<FChangelistTreeItem> ChangelistTreeItem;
+		TSharedPtr<FShelvedChangelistTreeItem> ShelvedChangelistTreeItem;
+		if (TSharedPtr<IChangelistTreeItem>* CachedChangelistItem = SourceControlItemCache.Find(ChangelistState))
+		{
+			check((*CachedChangelistItem)->GetTreeItemType() == IChangelistTreeItem::Changelist);
+			ChangelistTreeItem = StaticCastSharedPtr<FChangelistTreeItem>(*CachedChangelistItem);
+			ChangelistTreeItem->RemoveAllChildren(); // Will relink the children against the active filter below.
+			ChangelistTreeItem->ShelvedChangelistItem->RemoveAllChildren(); // Will relink the children against the active filter below.
+			OnSourceControlCachedItemVisited(ChangelistTreeItem.Get());
+		}
+		else
+		{
+			ChangelistTreeItem = MakeShared<FChangelistTreeItem>(ChangelistState);
+			ChangelistTreeItem->ShelvedChangelistItem = MakeShared<FShelvedChangelistTreeItem>();
+			OnNewSourceControlItemCreated(ChangelistState, ChangelistTreeItem);
+		}
+		ShelvedChangelistTreeItem = ChangelistTreeItem->ShelvedChangelistItem;
 
-		for (const TSharedRef<ISourceControlChangelistState>& ChangelistState : ChangelistsStates)
+		const bool bShelvedChangelistPassesFilter = ChangelistState->GetShelvedFilesStates().Num() > 0 && ChangelistTextFilter->PassesFilter(*ShelvedChangelistTreeItem);
+		const bool bShelvedChangelistMustBeDisplayed = bShelvedChangelistPassesFilter && !ChangelistTextFilter->GetRawFilterText().IsEmpty();
+		const bool bChangelistPassedFilter = bShelvedChangelistMustBeDisplayed || ChangelistTextFilter->PassesFilter(*ChangelistTreeItem);
+		if (bChangelistPassedFilter)
 		{
-			TSharedRef<IChangelistTreeItem> ChangelistNode = MakeShared<FChangelistTreeItem>(ChangelistState);
-			TSharedPtr<IChangelistTreeItem> ShelvedFilesNode;
-			bool bShelvedFilesNodeFilteredOut = true;
-			bool bSearchForcedChangelistNodeExpansion = false;
-			if (ChangelistState->GetShelvedFilesStates().Num() > 0)
+			ChangelistTreeNodes.Add(ChangelistTreeItem);
+			if (bShelvedChangelistPassesFilter)
 			{
-				ShelvedFilesNode = MakeShared<FShelvedChangelistTreeItem>();
-				bShelvedFilesNodeFilteredOut = !ChangelistTextFilter->PassesFilter(*ShelvedFilesNode);
-				bSearchForcedChangelistNodeExpansion = !bShelvedFilesNodeFilteredOut && !ChangelistTextFilter->GetRawFilterText().IsEmpty();
+				ChangelistTreeItem->AddChild(ShelvedChangelistTreeItem.ToSharedRef());
+				if (bShelvedChangelistMustBeDisplayed)
+				{
+					ChangelistTreeView->SetItemExpansion(ChangelistTreeItem, /*bShouldExpand*/true);
+				}
 			}
-		
-			// Check if the 'changelist' node must be kept.
-			if (!ChangelistTextFilter->PassesFilter(*ChangelistNode) && bShelvedFilesNodeFilteredOut)
+		}
+
+		// Create or reuse the 'File' nodes.
+		const bool bChangelistSelected = SelectedChangelistItem == ChangelistTreeItem;
+		for (const TSharedRef<ISourceControlState>& FileState : ChangelistState->GetFilesStates())
+		{
+			TSharedPtr<IFileViewTreeItem> FileItem;
+			if (TSharedPtr<IChangelistTreeItem>* CachedFileItem = SourceControlItemCache.Find(FileState))
 			{
-				continue; // Both the changelist and the shelved nodes are filtered out.
+				check((*CachedFileItem)->GetTreeItemType() == IChangelistTreeItem::File);
+				FileItem = StaticCastSharedPtr<FFileTreeItem>(*CachedFileItem);
+				OnSourceControlCachedItemVisited(FileItem.Get());
 			}
-
-			// Add the changelist.
-			ChangelistTreeNodes.Add(ChangelistNode);
-			if (bSearchForcedChangelistNodeExpansion)
+			else
 			{
-				ChangelistTreeView->SetItemExpansion(ChangelistNode, /*bShouldExpand*/true);
+				FileItem = MakeShared<FFileTreeItem>(FileState, bBeautifyPaths);
+				OnNewSourceControlItemCreated(FileState, FileItem);
 			}
 
-			// Add the changelist files.
-			for (const TSharedRef<ISourceControlState>& FileState : ChangelistState->GetFilesStates())
+			if (bChangelistPassedFilter && FileTextFilter->PassesFilter(*FileItem))
 			{
-				TSharedRef<FFileTreeItem> FileItem = MakeShared<FFileTreeItem>(FileState, bBeautifyPaths);
-				if (FileTextFilter->PassesFilter(*FileItem))
+				ChangelistTreeItem->AddChild(FileItem.ToSharedRef());
+				if (bChangelistSelected)
 				{
-					ChangelistNode->AddChild(FileItem);
-					SlowTask.EnterProgressFrame();
-					bBeautifyPaths &= !SlowTask.ShouldCancel();
+					AddItemToFileViewWrapper(FileItem);
 				}
 			}
+		}
 
-			if (ShelvedFilesNode)
+		// Create or reuse the 'Shelved File' nodes.
+		const bool bShelvedChangeslistSelected = SelectedChangelistItem == ShelvedChangelistTreeItem;
+		for (const TSharedRef<ISourceControlState>& ShelvedFileState : ChangelistState->GetShelvedFilesStates())
+		{
+			TSharedPtr<IFileViewTreeItem> ShelvedFileItem;
+			if (TSharedPtr<IChangelistTreeItem>* CachedShelvedFileItem = SourceControlItemCache.Find(ShelvedFileState))
+			{
+				check((*CachedShelvedFileItem)->GetTreeItemType() == IChangelistTreeItem::ShelvedFile);
+				ShelvedFileItem = StaticCastSharedPtr<FShelvedFileTreeItem>(*CachedShelvedFileItem);
+				OnSourceControlCachedItemVisited(ShelvedFileItem.Get());
+			}
+			else
 			{
-				// Add a shelved files node under the changelist node.
-				ChangelistNode->AddChild(ShelvedFilesNode.ToSharedRef());
+				ShelvedFileItem = MakeShared<FShelvedFileTreeItem>(ShelvedFileState, bBeautifyPaths);
+				OnNewSourceControlItemCreated(ShelvedFileState, ShelvedFileItem);
+			}
 
-				// Add the shelved files.
-				for (const TSharedRef<ISourceControlState>& ShelvedFileState : ChangelistState->GetShelvedFilesStates())
+			if (bChangelistPassedFilter && bShelvedChangelistPassesFilter && FileTextFilter->PassesFilter(*ShelvedFileItem))
+			{
+				ShelvedChangelistTreeItem->AddChild(ShelvedFileItem.ToSharedRef());
+				if (bShelvedChangeslistSelected)
 				{
-					TSharedRef<FShelvedFileTreeItem> ShelvedFileItem = MakeShared<FShelvedFileTreeItem>(ShelvedFileState, bBeautifyPaths);
-					if (FileTextFilter->PassesFilter(*ShelvedFileItem))
-					{
-						ShelvedFilesNode->AddChild(ShelvedFileItem);
-						SlowTask.EnterProgressFrame();
-						bBeautifyPaths &= !SlowTask.ShouldCancel();
-					}
+					AddItemToFileViewWrapper(ShelvedFileItem);
 				}
 			}
-
-			SlowTask.EnterProgressFrame();
-			bBeautifyPaths &= !SlowTask.ShouldCancel();
 		}
 	}
 
+	RemoveDiscardedSourceControlItems();
+	ChangelistTreeView->RequestTreeRefresh();
+}
+
+void SSourceControlChangelistsWidget::OnRefreshUncontrolledChangelistWidgets(int64 CurrUpdateNum, const TFunction<void(TSharedPtr<IFileViewTreeItem>&)>& AddItemToFileView)
+{
+	TRACE_CPUPROFILER_EVENT_SCOPE(SSourceControlChangelistsWidget::OnRefreshUncontrolledChangelistWidgets);
+
+	const bool bBeautifyPaths = true;
+	const int64 PrevUpdateNum = CurrUpdateNum - 1;
+	int32 DisplayedItemCount = 0;
+	int32 VisitedItemCount = 0;
+
+	TSharedPtr<IChangelistTreeItem> SelectedChangelistItem = UncontrolledChangelistTreeView->GetNumItemsSelected() > 0 ? UncontrolledChangelistTreeView->GetSelectedItems()[0] : TSharedPtr<IChangelistTreeItem>();
+
+	auto HasUncontrolledChangelistSelected = [&SelectedChangelistItem]()
+	{
+		return SelectedChangelistItem && SelectedChangelistItem->GetTreeItemType() == IChangelistTreeItem::UncontrolledChangelist;
+	};
+
+	auto OnUncontrolledChangelistCachedItemVisited = [CurrUpdateNum, &VisitedItemCount](IChangelistTreeItem* Item)
+	{
+		checkf(Item->VisitedUpdateNum < CurrUpdateNum, TEXT("The same uncontrolled item was visited twice. It is likely present in more than one uncontrolled changelist"));
+		Item->VisitedUpdateNum = CurrUpdateNum;
+		++VisitedItemCount;
+	};
+
+	auto OnNewUncontrolledChangelistItemCreated = [this, &OnUncontrolledChangelistCachedItemVisited](TSharedPtr<void> Key, TSharedPtr<IChangelistTreeItem> Item)
 	{
-		TRACE_CPUPROFILER_EVENT_SCOPE(SSourceControlChangelistsWidget::RefreshUncontrolledChangelists);
+		OnUncontrolledChangelistCachedItemVisited(Item.Get());
+		UncontrolledChangelistItemCache.Emplace(MoveTemp(Key), MoveTemp(Item));
+	};
+
+	auto OnNewOfflineFileCreated = [this, &OnUncontrolledChangelistCachedItemVisited](FString Key, TSharedPtr<IFileViewTreeItem> Item)
+	{
+		OnUncontrolledChangelistCachedItemVisited(Item.Get());
+		OfflineFileItemCache.Emplace(MoveTemp(Key), MoveTemp(Item));
+	};
+
+	auto AddItemToFileViewWrapper = [&AddItemToFileView, &DisplayedItemCount](TSharedPtr<IFileViewTreeItem>& Item)
+	{
+		++DisplayedItemCount;
+		AddItemToFileView(Item);
+	};
+
+	auto RemoveDiscardedUncontrolledChangelistItems = [this, VisitedItemCount, DisplayedItemCount, CurrUpdateNum, HasUncontrolledChangelistSelected]()
+	{
+		TRACE_CPUPROFILER_EVENT_SCOPE(SSourceControlChangelistsWidget::RemoveDiscardedUncontrolledChangelistItems);
+
+		check(UncontrolledChangelistItemCache.Num() + OfflineFileItemCache.Num() >= VisitedItemCount); // Cannot visit items that are not cached.
+
+		// Remove the items used to be visited but weren't visited during this iteration.
+		int32 DiscardCount = UncontrolledChangelistItemCache.Num() + OfflineFileItemCache.Num() - VisitedItemCount;
+		for (auto It = UncontrolledChangelistItemCache.CreateIterator(); It && DiscardCount > 0 ; ++It)
+		{
+			if (It->Value->VisitedUpdateNum < CurrUpdateNum)
+			{
+				It.RemoveCurrent();
+				--DiscardCount;
+			}
+		}
+		for (auto It = OfflineFileItemCache.CreateIterator(); It && DiscardCount > 0 ; ++It)
+		{
+			if (It->Value->VisitedUpdateNum < CurrUpdateNum)
+			{
+				It.RemoveCurrent();
+				--DiscardCount;
+			}
+		}
 
-		for (const TSharedRef<FUncontrolledChangelistState>& UncontrolledChangelistState : UncontrolledChangelistStates)
+		// If the uncontrolled changelist feeds the file view.
+		if (HasUncontrolledChangelistSelected())
 		{
-			// Add an uncontrolled changelist.
-			TSharedRef<IChangelistTreeItem> UncontrolledChangelistNode = MakeShared<FUncontrolledChangelistTreeItem>(UncontrolledChangelistState);
-			if (!UncontrolledChangelistTextFilter->PassesFilter(*UncontrolledChangelistNode))
+			// Unless the FileTreeNode array was reset because the selected changelist changed, it still contains the union of items to display for this iteration and the item that were displayed last
+			// iteration, so we need to remove item that aren't displayed anymore. (The FileTreeNode is not reset between updates to preserve the sort order as much as possible).
+			if (FileListNodes.Num() > DisplayedItemCount)
 			{
-				continue;
+				// Remove items that were not displayed anymore.
+				FileListNodes.RemoveAll([CurrUpdateNum](const TSharedPtr<IChangelistTreeItem>& Candidate) { return Candidate->DisplayedUpdateNum < CurrUpdateNum; });
 			}
+		}
+	};
+
+	FUncontrolledChangelistsModule& UncontrolledChangelistModule = FUncontrolledChangelistsModule::Get();
+	TArray<TSharedRef<FUncontrolledChangelistState>> UncontrolledChangelistStates = UncontrolledChangelistModule.GetChangelistStates();
+
+	UncontrolledChangelistTreeNodes.Reset(UncontrolledChangelistStates.Num());
+
+	for (const TSharedRef<FUncontrolledChangelistState>& UncontrolledChangelistState : UncontrolledChangelistStates)
+	{
+		// Create or reuse the 'Uncontrolled Changelist' node.
+		TSharedPtr<FUncontrolledChangelistTreeItem> UncontrolledChangelistNode;
+		if (TSharedPtr<IChangelistTreeItem>* CachedUncontrolledChangelistItem = UncontrolledChangelistItemCache.Find(UncontrolledChangelistState))
+		{
+			check((*CachedUncontrolledChangelistItem)->GetTreeItemType() == IChangelistTreeItem::UncontrolledChangelist);
+			UncontrolledChangelistNode = StaticCastSharedPtr<FUncontrolledChangelistTreeItem>(*CachedUncontrolledChangelistItem);
+			UncontrolledChangelistNode->RemoveAllChildren(); // Will relink the children against the active filter below.
+			OnUncontrolledChangelistCachedItemVisited(UncontrolledChangelistNode.Get());
+		}
+		else
+		{
+			UncontrolledChangelistNode = MakeShared<FUncontrolledChangelistTreeItem>(UncontrolledChangelistState);
+			OnNewUncontrolledChangelistItemCreated(UncontrolledChangelistState, UncontrolledChangelistNode);
+		}
+
+		const bool bChangelistPassedFilter = UncontrolledChangelistTextFilter->PassesFilter(*UncontrolledChangelistNode);
+		if (bChangelistPassedFilter)
+		{
 			UncontrolledChangelistTreeNodes.Add(UncontrolledChangelistNode);
+		}
 
-			for (const TSharedRef<ISourceControlState>& FileState : UncontrolledChangelistState->GetFilesStates())
+		// Create or reuse 'Uncontrolled' files.
+		const bool bUncontrolledChangelistSelected = SelectedChangelistItem == UncontrolledChangelistNode;
+		for (const TSharedRef<ISourceControlState>& FileState : UncontrolledChangelistState->GetFilesStates())
+		{
+			TSharedPtr<IFileViewTreeItem> UncontrolledFileItem;
+			if (TSharedPtr<IChangelistTreeItem>* CachedUncontrolledFileItem = UncontrolledChangelistItemCache.Find(FileState))
+			{
+				check((*CachedUncontrolledFileItem)->GetTreeItemType() == IChangelistTreeItem::File);
+				UncontrolledFileItem = StaticCastSharedPtr<FFileTreeItem>(*CachedUncontrolledFileItem);
+				OnUncontrolledChangelistCachedItemVisited(UncontrolledFileItem.Get());
+			}
+			else
+			{
+				UncontrolledFileItem = MakeShared<FFileTreeItem>(FileState, bBeautifyPaths);
+				OnNewUncontrolledChangelistItemCreated(FileState, UncontrolledFileItem);
+			}
+
+			if (bChangelistPassedFilter && FileTextFilter->PassesFilter(*UncontrolledFileItem))
 			{
-				TSharedRef<FFileTreeItem> UncontrolledFileItem = MakeShared<FFileTreeItem>(FileState, bBeautifyPaths);
-				if (FileTextFilter->PassesFilter(*UncontrolledFileItem))
+				UncontrolledChangelistNode->AddChild(UncontrolledFileItem.ToSharedRef());
+				if (bUncontrolledChangelistSelected)
 				{
-					UncontrolledChangelistNode->AddChild(UncontrolledFileItem);
-					SlowTask.EnterProgressFrame();
-					bBeautifyPaths &= !SlowTask.ShouldCancel();
+					AddItemToFileViewWrapper(UncontrolledFileItem);
 				}
 			}
+		}
+
+		// Create or reuse 'Offline' files.
+		for (const FString& Filename : UncontrolledChangelistState->GetOfflineFiles())
+		{
+			TSharedPtr<IFileViewTreeItem> OfflineFileItem;
+			if (TSharedPtr<IFileViewTreeItem>* CachedOfflineFileItem = OfflineFileItemCache.Find(Filename))
+			{
+				OfflineFileItem = *CachedOfflineFileItem;
+				OnUncontrolledChangelistCachedItemVisited(OfflineFileItem.Get());
+			}
+			else
+			{
+				OfflineFileItem = MakeShared<FOfflineFileTreeItem>(Filename);
+				OnNewOfflineFileCreated(Filename, OfflineFileItem);
+			}
 
-			for (const FString& Filename : UncontrolledChangelistState->GetOfflineFiles())
+			if (bChangelistPassedFilter && FileTextFilter->PassesFilter(*OfflineFileItem))
 			{
-				TSharedRef<FOfflineFileTreeItem> OfflineFileItem = MakeShared<FOfflineFileTreeItem>(Filename);
-				if (FileTextFilter->PassesFilter(*OfflineFileItem))
+				UncontrolledChangelistNode->AddChild(OfflineFileItem.ToSharedRef());
+				if (bUncontrolledChangelistSelected)
 				{
-					UncontrolledChangelistNode->AddChild(OfflineFileItem);
-					SlowTask.EnterProgressFrame();
-					bBeautifyPaths &= !SlowTask.ShouldCancel();
+					AddItemToFileViewWrapper(OfflineFileItem);
 				}
 			}
+		}
+	}
+
+	RemoveDiscardedUncontrolledChangelistItems();
+	UncontrolledChangelistTreeView->RequestTreeRefresh();
+}
 
-			SlowTask.EnterProgressFrame();
-			bBeautifyPaths &= !SlowTask.ShouldCancel();
+void SSourceControlChangelistsWidget::OnRefreshUI(ERefreshFlags RefreshFlag)
+{
+	TRACE_CPUPROFILER_EVENT_SCOPE(SSourceControlChangelistsWidget::OnRefreshUI);
+
+	// The refresh algorithm tries to only update what changed. This is required to deal smoothly with reasonably large
+	// changelists (5000 - 10000 files). The code focus on:
+	//   - Detecting added and removed items to do a 'delta' update (important for the file view)
+	//   - Preserving the sort order when possible.
+	//   - Caching and preserving file timestamps and update them using background task.
+	//   - Try performing expensive operations only on visible items.
+	// This gives those benefits:
+	// - The UI doesn't have to regenerate widget that didn't change.
+	// - The selection of items that were displayed before and after an update is naturally preserved.
+	// - The sort order is preserved most of the time.
+	// - The cache keeps existing item even when filtered out, especially to avoiding to query the file timestamps (very expensive) too often.
+
+	if (!AreChangelistsEnabled())
+	{
+		ClearChangelistsTree();
+		return;
+	}
+
+	const int64 PrevUpdateNum = UpdateRequestNum;
+	const int64 CurrUpdateNum = ++UpdateRequestNum;
+	int32 NewDisplayedItemCount = 0;
+	bool bDisplayedIconsPriorityChanged = false;
+
+	auto AddItemToFileView = [this, &NewDisplayedItemCount, &bDisplayedIconsPriorityChanged, PrevUpdateNum, CurrUpdateNum](TSharedPtr<IFileViewTreeItem>& Item)
+	{
+		// If the items wasn't displayed last update.
+		if (Item->DisplayedUpdateNum != PrevUpdateNum)
+		{
+			checkfSlow(!FileListNodes.Contains(Item), TEXT("Inserting duplicated items. Something is wrong with the display update number."))
+
+			++NewDisplayedItemCount;
+			bUpdateMonitoredFileStatusList = true;
+
+			// Add item at the end, will need to sort the list if the view is sorted.
+			FileListNodes.Emplace(Item);
+
+			// That's the first time the item is displayed.
+			if (Item->DisplayedUpdateNum == -1)
+			{
+				Item->SetLastModifiedDateTime(IFileManager::Get().GetTimeStamp(*Item->GetFullPathname()));
+			}
 		}
+
+		// Check when the file status changes (possibly invalidating sorting)
+		int32 IconSortingPriority = Item->GetIconSortingPriority();
+		if (Item->DisplayedIconPriority != IconSortingPriority)
+		{
+			bDisplayedIconsPriorityChanged = true;
+			Item->DisplayedIconPriority = IconSortingPriority;
+		}
+
+		// Mark this item as 'displayed' at this 'iteration'.
+		Item->DisplayedUpdateNum = CurrUpdateNum;
+	};
+
+	bool bFileViewUpdated = false;
+	if (EnumHasAnyFlags(RefreshFlag, ERefreshFlags::SourceControlChangelists))
+	{
+		OnRefreshSourceControlWidgets(CurrUpdateNum, AddItemToFileView);
+		bFileViewUpdated |= (ChangelistTreeView->GetNumItemsSelected() > 0);
+	}
+	if (EnumHasAnyFlags(RefreshFlag, ERefreshFlags::UncontrolledChangelists))
+	{
+		OnRefreshUncontrolledChangelistWidgets(CurrUpdateNum, AddItemToFileView);
+		bFileViewUpdated |= (UncontrolledChangelistTreeView->GetNumItemsSelected() > 0);
 	}
 
-	// Views were rebuilt from scratch, try expanding and selecting the nodes that were in that state before the update.
-	RestoreExpandedAndSelectionStates(ExpandedAndSelectedStates);
+	if (!bFileViewUpdated)
+	{
+		// The update didn't touch the file view, but the files (if any) were still displayed during this update interation, bump their displayed update number.
+		for (TSharedPtr<IChangelistTreeItem>& Item : FileListNodes)
+		{
+			Item->DisplayedUpdateNum = CurrUpdateNum;
+		}
+	}
 
 	if (FilesToSelect.Num() > 0)
 	{
@@ -912,12 +1277,13 @@ void SSourceControlChangelistsWidget::OnRefresh()
 		SetSelectedFiles(LocalFilesToSelect);
 	}
 
-	// Restore the sort order.
-	SortFileView();
+	const bool bNeedSorting = !PrimarySortedColumn.IsNone() && (NewDisplayedItemCount > 0 || (IsFileViewSortedByFileStatusIcon() && bDisplayedIconsPriorityChanged));
+	if (bNeedSorting)
+	{
+		SortFileView();
+	}
 
-	ChangelistTreeView->RequestTreeRefresh();
-	UncontrolledChangelistTreeView->RequestTreeRefresh();
-	FileTreeView->RequestTreeRefresh();
+	FileListView->RequestListRefresh();
 }
 
 void SSourceControlChangelistsWidget::OnSourceControlProviderChanged(ISourceControlProvider& OldProvider, ISourceControlProvider& NewProvider)
@@ -927,6 +1293,18 @@ void SSourceControlChangelistsWidget::OnSourceControlProviderChanged(ISourceCont
 
 	bSourceControlAvailable = NewProvider.IsAvailable(); // Check if it is connected.
 	bShouldRefresh = true;
+
+	if (&NewProvider != &OldProvider)
+	{
+		if (ChangelistTreeView->GetNumItemsSelected() > 0)
+		{
+			FileListNodes.Reset();
+			FileListView->RequestListRefresh();
+		}
+		ChangelistTreeNodes.Reset();
+		SourceControlItemCache.Reset();
+		ChangelistTreeView->RequestTreeRefresh();
+	}
 }
 
 void SSourceControlChangelistsWidget::OnSourceControlStateChanged()
@@ -935,7 +1313,7 @@ void SSourceControlChangelistsWidget::OnSourceControlStateChanged()
 
 	// NOTE: No need to call RequestChangelistsRefresh() to force the SCC to update internal states. We are being invoked because it was update, we just
 	//       need to update the UI to reflect those state changes.
-	OnRefresh();
+	OnRefreshUI(ERefreshFlags::SourceControlChangelists);
 }
 
 void SSourceControlChangelistsWidget::OnChangelistsStatusUpdated(const TSharedRef<ISourceControlOperation>& InOperation, ECommandResult::Type InType)
@@ -943,46 +1321,11 @@ void SSourceControlChangelistsWidget::OnChangelistsStatusUpdated(const TSharedRe
 	// NOTE: This is invoked when the 'FUpdatePendingChangelistsStatus' completes. No need to refresh the tree views because OnSourceControlStateChanged() is also called.
 	OnEndSourceControlOperation(InOperation, InType);
 	EndRefreshStatus();
-
-	// Check if a changelist is selected.
-	TArray<TSharedPtr<IChangelistTreeItem>> SelectedChangelists = ChangelistTreeView->GetSelectedItems();
-	if (SelectedChangelists.Num() == 0)
-	{
-		// Check if an uncontrolled CL is selected. A changelist and an uncontrolled changelist cannot be selected at the same time.
-		SelectedChangelists = UncontrolledChangelistTreeView->GetSelectedItems();
-	}
-
-	// If there is a selected changelist (controlled or uncontrolled), refresh the file status with respect to the source control. The response will come
-	// as an update through OnSourceControlStateChanged() callback.
-	if (SelectedChangelists.Num() > 0)
-	{
-		RequestFileStatusRefresh(*SelectedChangelists[0]);
-	}
 }
 
-void SChangelistTree::Private_SetItemSelection(FChangelistTreeItemPtr TheItem, bool bShouldBeSelected, bool bWasUserDirected)
+void SSourceControlChangelistsWidget::OnUncontrolledChangelistStateChanged()
 {
-	bool bAllowSelectionChange = true;
-
-	if (bShouldBeSelected && !SelectedItems.IsEmpty())
-	{
-		// Prevent selecting changelists and files at the same time.
-		FChangelistTreeItemPtr CurrentlySelectedItem = (*SelectedItems.begin());
-		if (TheItem->GetTreeItemType() != CurrentlySelectedItem->GetTreeItemType())
-		{
-			bAllowSelectionChange = false;
-		}
-		// Prevent selecting items that don't share the same root
-		else if (TheItem->GetParent() != CurrentlySelectedItem->GetParent())
-		{
-			bAllowSelectionChange = false;
-		}
-	}
-
-	if (bAllowSelectionChange)
-	{
-		STreeView::Private_SetItemSelection(TheItem, bShouldBeSelected, bWasUserDirected);
-	}
+	OnRefreshUI(ERefreshFlags::UncontrolledChangelists);
 }
 
 FSourceControlChangelistStatePtr SSourceControlChangelistsWidget::GetCurrentChangelistState()
@@ -1121,71 +1464,69 @@ void SSourceControlChangelistsWidget::SetSelectedFiles(const TArray<FString>& Fi
 		}
 	}
 
-	// If we found a Changelist, select files 
+	// Clear the current selection.
+	FileListView->ClearSelection();
+	TSharedPtr<IChangelistTreeItem> LastFoundItem;
+
+	// If we found a Changelist, select files.
 	if (FoundChangelistTreeItem)
 	{
-		// Save current Selection/Expansion state
-		FExpandedAndSelectionStates State;
-		SaveExpandedAndSelectionStates(State);
-		State.SelectedChangelistNode = nullptr;
-		State.SelectedUncontrolledChangelistNode = nullptr;
-		State.SelectedFileNodes.Empty();
-
 		// To make search faster store all filenames lower case
 		TSet<FString> FilenamesLowerCase;
 		Algo::Transform(Filenames, FilenamesLowerCase, [](const FString& Filename) { return Filename.ToLower(); });
 
-		if (FoundChangelistTreeItem->GetTreeItemType() == IChangelistTreeItem::Changelist)
+		for (const TSharedPtr<IChangelistTreeItem>& ChildItem : FoundChangelistTreeItem->GetChildren())
 		{
-			State.SelectedChangelistNode = StaticCastSharedPtr<FChangelistTreeItem>(FoundChangelistTreeItem);
-
-			for (const TSharedPtr<IChangelistTreeItem>& ChildItem : FoundChangelistTreeItem->GetChildren())
+			if (ChildItem->GetTreeItemType() == IChangelistTreeItem::File)
 			{
-				if (ChildItem->GetTreeItemType() == IChangelistTreeItem::File)
+				const FString& ChildFilename = StaticCastSharedPtr<FFileTreeItem>(ChildItem)->FileState->GetFilename().ToLower();
+				if (FilenamesLowerCase.Contains(ChildFilename))
 				{
-					const FString& ChildFilename = StaticCastSharedPtr<FFileTreeItem>(ChildItem)->FileState->GetFilename().ToLower();
-					if (FilenamesLowerCase.Contains(ChildFilename))
-					{
-						State.SelectedFileNodes.Add(ChildItem);
-					}
+					FileListView->SetItemSelection(ChildItem, /*bSelected*/true);
+					LastFoundItem = ChildItem;
+				}
+			}
+			else if (ChildItem->GetTreeItemType() == IChangelistTreeItem::OfflineFile)
+			{
+				const FString& ChildFilename = StaticCastSharedPtr<FOfflineFileTreeItem>(ChildItem)->GetFilename().ToLower();
+				if (FilenamesLowerCase.Contains(ChildFilename))
+				{
+					FileListView->SetItemSelection(ChildItem, /*bSelected*/true);
+					LastFoundItem = ChildItem;
 				}
 			}
 		}
-		else if (FoundChangelistTreeItem->GetTreeItemType() == IChangelistTreeItem::UncontrolledChangelist)
-		{
-			State.SelectedUncontrolledChangelistNode = StaticCastSharedPtr<FUncontrolledChangelistTreeItem>(FoundChangelistTreeItem);
 
-			for (const TSharedPtr<IChangelistTreeItem>& ChildItem : FoundChangelistTreeItem->GetChildren())
+		if (LastFoundItem)
+		{
+			if (FoundChangelistTreeItem->GetTreeItemType() == IChangelistTreeItem::Changelist)
 			{
-				if (ChildItem->GetTreeItemType() == IChangelistTreeItem::File)
+				// Ensure the area is expanded.
+				ChangelistExpandableArea->SetExpanded(true);
+
+				if (ChangelistTreeView->GetNumItemsSelected() == 0 || ChangelistTreeView->GetSelectedItems()[0] != FoundChangelistTreeItem)
 				{
-					const FString& ChildFilename = StaticCastSharedPtr<FFileTreeItem>(ChildItem)->FileState->GetFilename().ToLower();
-					if (FilenamesLowerCase.Contains(ChildFilename))
-					{
-						State.SelectedFileNodes.Add(ChildItem);
-					}
+					ChangelistTreeView->SetSelection(FoundChangelistTreeItem);
 				}
-				else if (ChildItem->GetTreeItemType() == IChangelistTreeItem::OfflineFile)
+			}
+			else if (FoundChangelistTreeItem->GetTreeItemType() == IChangelistTreeItem::UncontrolledChangelist)
+			{
+				// Ensure the area is expanded.
+				UncontrolledChangelistExpandableArea->SetExpanded(true);
+
+				if (UncontrolledChangelistTreeView->GetNumItemsSelected() == 0 || UncontrolledChangelistTreeView->GetSelectedItems()[0] != FoundChangelistTreeItem)
 				{
-					const FString& ChildFilename = StaticCastSharedPtr<FOfflineFileTreeItem>(ChildItem)->GetFilename().ToLower();
-					if (FilenamesLowerCase.Contains(ChildFilename))
-					{
-						State.SelectedFileNodes.Add(ChildItem);
-					}
+					UncontrolledChangelistTreeView->SetSelection(FoundChangelistTreeItem);
 				}
 			}
+			FileListView->RequestScrollIntoView(LastFoundItem);
 		}
-
-		RestoreExpandedAndSelectionStates(State);
-
-		// Restore the sort order.
-		SortFileView();
 	}
 }
 
 bool SSourceControlChangelistsWidget::HasFilesSelected() const
 {
-	if (FileTreeView->GetNumItemsSelected() == 0)
+	if (FileListView->GetNumItemsSelected() == 0)
 	{
 		return false;
 	}
@@ -1198,7 +1539,7 @@ bool SSourceControlChangelistsWidget::HasFilesSelected() const
 	case IChangelistTreeItem::Changelist:
 		return true; // If this type of changeslist is selected and the file view as item selected, it means files are selected.
 	case IChangelistTreeItem::UncontrolledChangelist:
-		for (const TSharedPtr<IChangelistTreeItem>& Item : FileTreeView->GetSelectedItems())
+		for (const TSharedPtr<IChangelistTreeItem>& Item : FileListView->GetSelectedItems())
 		{
 			if (Item->GetTreeItemType() == IChangelistTreeItem::File)
 			{
@@ -1215,7 +1556,7 @@ bool SSourceControlChangelistsWidget::HasFilesSelected() const
 TArray<FString> SSourceControlChangelistsWidget::GetSelectedFiles()
 {
 	TArray<FString> Files;
-	for (const TSharedPtr<IChangelistTreeItem>& Item : FileTreeView->GetSelectedItems())
+	for (const TSharedPtr<IChangelistTreeItem>& Item : FileListView->GetSelectedItems())
 	{
 		if (Item->GetTreeItemType() == IChangelistTreeItem::File)
 		{
@@ -1228,7 +1569,7 @@ TArray<FString> SSourceControlChangelistsWidget::GetSelectedFiles()
 
 void SSourceControlChangelistsWidget::GetSelectedFiles(TArray<FString>& OutControlledFiles, TArray<FString>& OutUncontrolledFiles)
 {
-	for (const TSharedPtr<IChangelistTreeItem>& Item : FileTreeView->GetSelectedItems())
+	for (const TSharedPtr<IChangelistTreeItem>& Item : FileListView->GetSelectedItems())
 	{
 		if (Item->GetTreeItemType() == IChangelistTreeItem::File)
 		{
@@ -1262,7 +1603,7 @@ void SSourceControlChangelistsWidget::GetSelectedFiles(TArray<FString>& OutContr
 
 void SSourceControlChangelistsWidget::GetSelectedFileStates(TArray<FSourceControlStateRef>& OutControlledFileStates, TArray<FSourceControlStateRef>& OutUncontrolledFileStates)
 {
-	TArray<TSharedPtr<IChangelistTreeItem>> SelectedItems = FileTreeView->GetSelectedItems();
+	TArray<TSharedPtr<IChangelistTreeItem>> SelectedItems = FileListView->GetSelectedItems();
 
 	for (const TSharedPtr<IChangelistTreeItem>& Item : SelectedItems)
 	{
@@ -1287,7 +1628,7 @@ void SSourceControlChangelistsWidget::GetSelectedFileStates(TArray<FSourceContro
 
 bool SSourceControlChangelistsWidget::HasShelvedFilesSelected() const
 {
-	return FileTreeView->GetNumItemsSelected() > 0 &&
+	return FileListView->GetNumItemsSelected() > 0 &&
 	       ChangelistTreeView->GetNumItemsSelected() > 0 &&
 	       ChangelistTreeView->GetSelectedItems()[0]->GetTreeItemType() == IChangelistTreeItem::ShelvedChangelist;
 }
@@ -1296,7 +1637,7 @@ TArray<FString> SSourceControlChangelistsWidget::GetSelectedShelvedFiles()
 {
 	TArray<FString> ShelvedFiles;
 
-	for (const TSharedPtr<IChangelistTreeItem>& Item : FileTreeView->GetSelectedItems())
+	for (const TSharedPtr<IChangelistTreeItem>& Item : FileListView->GetSelectedItems())
 	{
 		if (Item->GetTreeItemType() == IChangelistTreeItem::ShelvedFile)
 		{
@@ -2332,21 +2673,21 @@ void SSourceControlChangelistsWidget::OnDiffAgainstDepot()
 
 bool SSourceControlChangelistsWidget::CanDiffAgainstDepot()
 {
-	return FileTreeView->GetNumItemsSelected() == 1 && HasFilesSelected();
+	return FileListView->GetNumItemsSelected() == 1 && HasFilesSelected();
 }
 
 void SSourceControlChangelistsWidget::OnDiffAgainstWorkspace()
 {
 	if (GetSelectedShelvedFiles().Num() > 0)
 	{
-		FSourceControlStateRef FileState = StaticCastSharedPtr<FShelvedFileTreeItem>(FileTreeView->GetSelectedItems()[0])->FileState;
+		FSourceControlStateRef FileState = StaticCastSharedPtr<FShelvedFileTreeItem>(FileListView->GetSelectedItems()[0])->FileState;
 		FSourceControlWindows::DiffAgainstShelvedFile(FileState);
 	}
 }
 
 bool SSourceControlChangelistsWidget::CanDiffAgainstWorkspace()
 {
-	return FileTreeView->GetNumItemsSelected() == 1 && HasShelvedFilesSelected();
+	return FileListView->GetNumItemsSelected() == 1 && HasShelvedFilesSelected();
 }
 
 TSharedPtr<SWidget> SSourceControlChangelistsWidget::OnOpenContextMenu()
@@ -2560,9 +2901,9 @@ TSharedPtr<SWidget> SSourceControlChangelistsWidget::OnOpenContextMenu()
 	return ToolMenus->GenerateWidget(Menu);
 }
 
-TSharedRef<SChangelistTree> SSourceControlChangelistsWidget::CreateChangelistTreeView(TArray<TSharedPtr<IChangelistTreeItem>>& ItemSources)
+TSharedRef<STreeView<FChangelistTreeItemPtr>> SSourceControlChangelistsWidget::CreateChangelistTreeView(TArray<TSharedPtr<IChangelistTreeItem>>& ItemSources)
 {
-	return SNew(SChangelistTree)
+	return SNew(STreeView<FChangelistTreeItemPtr>)
 		.ItemHeight(24.0f)
 		.TreeItemsSource(&ItemSources)
 		.OnGenerateRow(this, &SSourceControlChangelistsWidget::OnGenerateRow)
@@ -2573,7 +2914,7 @@ TSharedRef<SChangelistTree> SSourceControlChangelistsWidget::CreateChangelistTre
 		.OnSelectionChanged(this, &SSourceControlChangelistsWidget::OnChangelistSelectionChanged);
 }
 
-TSharedRef<STreeView<FChangelistTreeItemPtr>> SSourceControlChangelistsWidget::CreateChangelistFilesView()
+TSharedRef<SListView<FChangelistTreeItemPtr>> SSourceControlChangelistsWidget::CreateChangelistFilesView()
 {
 	USourceControlSettings* Settings = GetMutableDefault<USourceControlSettings>();
 	if (!Settings->bShowAssetTypeColumn)
@@ -2589,14 +2930,15 @@ TSharedRef<STreeView<FChangelistTreeItemPtr>> SSourceControlChangelistsWidget::C
 		FileViewHiddenColumnsList.Add(SourceControlFileViewColumn::CheckedOutByUser::Id());
 	}
 
-	TSharedRef<STreeView<FChangelistTreeItemPtr>> FileView = SNew(STreeView<FChangelistTreeItemPtr>)
+	TSharedRef<SListView<FChangelistTreeItemPtr>> FileView = SNew(SListView<FChangelistTreeItemPtr>)
 		.ItemHeight(24.0f)
-		.TreeItemsSource(&FileTreeNodes)
+		.ListItemsSource(&FileListNodes)
 		.OnGenerateRow(this, &SSourceControlChangelistsWidget::OnGenerateRow)
-		.OnGetChildren(this, &SSourceControlChangelistsWidget::OnGetFileChildren)
 		.SelectionMode(ESelectionMode::Multi)
 		.OnContextMenuOpening(this, &SSourceControlChangelistsWidget::OnOpenContextMenu)
 		.OnMouseButtonDoubleClick(this, &SSourceControlChangelistsWidget::OnItemDoubleClicked)
+		.OnListViewScrolled_Lambda([this](double InScrollOffset) { if (!IsFileViewSortedByFileStatusIcon()) { bUpdateMonitoredFileStatusList = true; } }) // If sorted by status icon, the full list of file is already monitored.
+		.OnItemToString_Debug_Lambda([this](FChangelistTreeItemPtr Item) { return static_cast<IFileViewTreeItem*>(Item.Get())->GetName(); })
 		.HeaderRow
 		(
 			SNew(SHeaderRow)
@@ -2679,14 +3021,14 @@ TSharedRef<STreeView<FChangelistTreeItemPtr>> SSourceControlChangelistsWidget::C
 
 void SSourceControlChangelistsWidget::OnFileViewHiddenColumnsListChanged()
 {
-	if (FileTreeView && FileTreeView->GetHeaderRow())
+	if (FileListView && FileListView->GetHeaderRow())
 	{
 		USourceControlSettings* Settings = GetMutableDefault<USourceControlSettings>();
 		Settings->bShowAssetTypeColumn = true;
 		Settings->bShowAssetCheckedOutByColumn = true;
 		Settings->bShowAssetLastModifiedTimeColumn = true;
 
-		for (const FName& ColumnId : FileTreeView->GetHeaderRow()->GetHiddenColumnIds())
+		for (const FName& ColumnId : FileListView->GetHeaderRow()->GetHiddenColumnIds())
 		{
 			if (ColumnId == SourceControlFileViewColumn::Type::Id())
 			{
@@ -2713,21 +3055,23 @@ TSharedRef<ITableRow> SSourceControlChangelistsWidget::OnGenerateRow(TSharedPtr<
 		return SNew(SChangelistTableRow, OwnerTable)
 			.TreeItemToVisualize(InTreeItem)
 			.HighlightText_Lambda([this]() { return ChangelistExpandableArea->GetSearchedText(); })
-			.OnPostDrop(this, &SSourceControlChangelistsWidget::OnRefresh);
+			.OnPostDrop_Lambda([this] { OnRefreshUI(ERefreshFlags::All); });
 
 	case IChangelistTreeItem::UncontrolledChangelist:
 		return SNew(SUncontrolledChangelistTableRow, OwnerTable)
 			.TreeItemToVisualize(InTreeItem)
 			.HighlightText_Lambda([this]() { return UncontrolledChangelistExpandableArea->GetSearchedText(); })
-			.OnPostDrop(this, &SSourceControlChangelistsWidget::OnRefresh);
+			.OnPostDrop_Lambda([this] { OnRefreshUI(ERefreshFlags::All); });
 
 	case IChangelistTreeItem::File:
+		bUpdateMonitoredFileStatusList = true;
 		return SNew(SFileTableRow, OwnerTable)
 			.TreeItemToVisualize(InTreeItem)
 			.HighlightText_Lambda([this]() { return FileSearchBox->GetText(); })
 			.OnDragDetected(this, &SSourceControlChangelistsWidget::OnFilesDragged);
 
 	case IChangelistTreeItem::OfflineFile:
+		bUpdateMonitoredFileStatusList = true;
 		return SNew(SOfflineFileTableRow, OwnerTable)
 			.TreeItemToVisualize(InTreeItem)
 			.HighlightText_Lambda([this]() { return FileSearchBox->GetText(); });
@@ -2738,6 +3082,7 @@ TSharedRef<ITableRow> SSourceControlChangelistsWidget::OnGenerateRow(TSharedPtr<
 			.HighlightText_Lambda([this]() { return ChangelistExpandableArea->GetSearchedText(); });
 
 	case IChangelistTreeItem::ShelvedFile:
+		bUpdateMonitoredFileStatusList = true;
 		return SNew(SFileTableRow, OwnerTable)
 			.TreeItemToVisualize(InTreeItem)
 			.HighlightText_Lambda([this]() { return FileSearchBox->GetText(); });
@@ -2751,11 +3096,11 @@ TSharedRef<ITableRow> SSourceControlChangelistsWidget::OnGenerateRow(TSharedPtr<
 
 FReply SSourceControlChangelistsWidget::OnFilesDragged(const FGeometry& InGeometry, const FPointerEvent& InMouseEvent)
 {
-	if (InMouseEvent.IsMouseButtonDown(EKeys::LeftMouseButton) && !FileTreeView->GetSelectedItems().IsEmpty())
+	if (InMouseEvent.IsMouseButtonDown(EKeys::LeftMouseButton) && !FileListView->GetSelectedItems().IsEmpty())
 	{
 		TSharedRef<FSCCFileDragDropOp> Operation = MakeShared<FSCCFileDragDropOp>();
 
-		for (const FChangelistTreeItemPtr& InTreeItem : FileTreeView->GetSelectedItems())
+		for (const FChangelistTreeItemPtr& InTreeItem : FileListView->GetSelectedItems())
 		{
 			if (InTreeItem->GetTreeItemType() == IChangelistTreeItem::File)
 			{
@@ -2781,27 +3126,14 @@ FReply SSourceControlChangelistsWidget::OnFilesDragged(const FGeometry& InGeomet
 	return FReply::Unhandled();
 }
 
-void SSourceControlChangelistsWidget::OnGetFileChildren(TSharedPtr<IChangelistTreeItem> InParent, TArray<FChangelistTreeItemPtr>& OutChildren)
-{
-	// Files are leave and don't have children.
-}
-
 void SSourceControlChangelistsWidget::OnGetChangelistChildren(FChangelistTreeItemPtr InParent, TArray<FChangelistTreeItemPtr>& OutChildren)
 {
 	if (InParent->GetTreeItemType() == IChangelistTreeItem::Changelist)
 	{
-		// In the data model, a changelist has files as children, but in UI, only the 'Shelved Files' node is displayed under the changelist,
-		// and the files are displayed in the file view at the right.
-		for (const TSharedPtr<IChangelistTreeItem>& Child : InParent->GetChildren())
+		const FChangelistTreeItem* ChangelistItem = static_cast<const FChangelistTreeItem*>(InParent.Get());
+		if (ChangelistItem->GetShelvedFileCount() > 0 && ChangelistItem->ShelvedChangelistItem->GetParent() != nullptr)
 		{
-			if (Child->GetTreeItemType() == IChangelistTreeItem::ShelvedChangelist)
-			{
-				if (Child->GetChildren().Num() > 0)
-				{
-					OutChildren.Add(Child); // Add the 'Shelved Files' only if there are shelved files.
-					break; // Found the only possible child for the UI.
-				}
-			}
+			OutChildren.Add(ChangelistItem->ShelvedChangelistItem); // Add the 'Shelved Files' only if is is linked to its parent (has items and is not filtered out).
 		}
 	}
 	else if (InParent->GetTreeItemType() == IChangelistTreeItem::UncontrolledChangelist)
@@ -2842,20 +3174,8 @@ void SSourceControlChangelistsWidget::OnChangelistSelectionChanged(TSharedPtr<IC
 {
 	TRACE_CPUPROFILER_EVENT_SCOPE(SSourceControlChangelistsWidget::OnChangelistSelectionChanged);
 
-	FileTreeNodes.Reset();
-
-	// Add the children of the parent item to the file tree node.
-	auto AddChangelistFilesToFileView = [this](TSharedPtr<IChangelistTreeItem> ParentItem, IChangelistTreeItem::TreeItemType DesiredChildrenType)
-	{
-		for (const TSharedPtr<IChangelistTreeItem>& Child : ParentItem->GetChildren())
-		{
-			if (Child->GetTreeItemType() == DesiredChildrenType)
-			{
-				FileTreeNodes.Add(Child);
-				UpdateLastModifiedTimestamp(*Child);
-			}
-		}
-	};
+	FileListNodes.Reset();
+	bUpdateMonitoredFileStatusList = true;
 
 	if (SelectedChangelist) // Can be a Changelist, Uncontrolled Changelist or Shelved Changelist
 	{
@@ -2865,31 +3185,18 @@ void SSourceControlChangelistsWidget::OnChangelistSelectionChanged(TSharedPtr<IC
 			case IChangelistTreeItem::Changelist:
 			case IChangelistTreeItem::ShelvedChangelist:
 				UncontrolledChangelistTreeView->ClearSelection(); // Don't have a changelists selected at the same time than an uncontrolled one, they share the same file view.
-				AddChangelistFilesToFileView(SelectedChangelist, ChangelistType == IChangelistTreeItem::Changelist ? IChangelistTreeItem::File : IChangelistTreeItem::ShelvedFile);
+				OnRefreshUI(ERefreshFlags::SourceControlChangelists);
 				break;
 
 			case IChangelistTreeItem::UncontrolledChangelist:
 				ChangelistTreeView->ClearSelection();
-				AddChangelistFilesToFileView(SelectedChangelist, IChangelistTreeItem::File);
-				AddChangelistFilesToFileView(SelectedChangelist, IChangelistTreeItem::OfflineFile);
+				OnRefreshUI(ERefreshFlags::UncontrolledChangelists);
 				break;
 
 			default:
 				break;
 		}
-
-		// If the user clicked the changelist himself (not a side effect of a refresh coming from a source control state change).
-		if (SelectionType != ESelectInfo::Direct)
-		{
-			// Fire a request to get more info about the  file status (checkout by, outdated, etc). This is expensive on the server to call periodically.
-			RequestFileStatusRefresh(*SelectedChangelist);
-
-			// Sort the view as it just changed.
-			SortFileView();
-		}
 	}
-
-	FileTreeView->RequestTreeRefresh();
 }
 
 EColumnSortPriority::Type SSourceControlChangelistsWidget::GetColumnSortPriority(const FName ColumnId) const
@@ -2939,15 +3246,30 @@ void SSourceControlChangelistsWidget::OnColumnSortModeChanged(const EColumnSortP
 		SecondarySortMode = InSortMode;
 	}
 
+	if (IsFileViewSortedByFileStatusIcon())
+	{
+		bUpdateMonitoredFileStatusList = true;
+	}
+
 	SortFileView();
-	FileTreeView->RequestListRefresh();
+	FileListView->RequestListRefresh();
+}
+
+bool SSourceControlChangelistsWidget::IsFileViewSortedByFileStatusIcon() const
+{
+	return PrimarySortedColumn == SourceControlFileViewColumn::Icon::Id() || SecondarySortedColumn == SourceControlFileViewColumn::Icon::Id();
+};
+
+bool SSourceControlChangelistsWidget::IsFileViewSortedByLastModifiedTimestamp() const
+{
+	return PrimarySortedColumn == SourceControlFileViewColumn::LastModifiedTimestamp::Id() || SecondarySortedColumn == SourceControlFileViewColumn::LastModifiedTimestamp::Id();
 }
 
 void SSourceControlChangelistsWidget::SortFileView()
 {
 	TRACE_CPUPROFILER_EVENT_SCOPE(SSourceControlChangelistsWidget::SortFileView);
 
-	if (PrimarySortedColumn.IsNone() || FileTreeNodes.IsEmpty())
+	if (PrimarySortedColumn.IsNone() || FileListNodes.IsEmpty())
 	{
 		return; // No column selected for sorting or nothing to sort.
 	}
@@ -3031,7 +3353,7 @@ void SSourceControlChangelistsWidget::SortFileView()
 		// NOTE: StableSort() would give a better experience when the sorted columns(s) has the same values and new values gets added, but it is slower
 		//       with large changelists (7600 items was about 1.8x slower in average measured with Unreal Insight). Because this code runs in the main
 		//       thread and can be invoked a lot, the trade off went if favor of speed.
-		FileTreeNodes.Sort([this, &PrimaryCompare, &SecondaryCompare](const TSharedPtr<IChangelistTreeItem>& Lhs, const TSharedPtr<IChangelistTreeItem>& Rhs)
+		FileListNodes.Sort([this, &PrimaryCompare, &SecondaryCompare](const TSharedPtr<IChangelistTreeItem>& Lhs, const TSharedPtr<IChangelistTreeItem>& Rhs)
 		{
 			int32 Result = PrimaryCompare(static_cast<IFileViewTreeItem*>(Lhs.Get()), static_cast<IFileViewTreeItem*>(Rhs.Get()));
 			if (Result < 0)
@@ -3054,7 +3376,7 @@ void SSourceControlChangelistsWidget::SortFileView()
 	}
 	else
 	{
-		FileTreeNodes.Sort([this, &PrimaryCompare, &SecondaryCompare](const TSharedPtr<IChangelistTreeItem>& Lhs, const TSharedPtr<IChangelistTreeItem>& Rhs)
+		FileListNodes.Sort([this, &PrimaryCompare, &SecondaryCompare](const TSharedPtr<IChangelistTreeItem>& Lhs, const TSharedPtr<IChangelistTreeItem>& Rhs)
 		{
 			int32 Result = PrimaryCompare(static_cast<IFileViewTreeItem*>(Lhs.Get()), static_cast<IFileViewTreeItem*>(Rhs.Get()));
 			if (Result > 0)
@@ -3125,158 +3447,9 @@ void SSourceControlChangelistsWidget::PopulateItemSearchStrings(const IChangelis
 	}
 }
 
-void SSourceControlChangelistsWidget::SaveExpandedAndSelectionStates(FExpandedAndSelectionStates& OutStates)
+void SSourceControlChangelistsWidget::OnPackageSaved(const FString& PackageFilename, UPackage* Package, FObjectPostSaveContext ObjectSaveContext)
 {
-	// Save the selected item from the 'changelists' tree.
-	TArray<TSharedPtr<IChangelistTreeItem>> SelectedChangelistItems = ChangelistTreeView->GetSelectedItems();
-	OutStates.SelectedChangelistNode = SelectedChangelistItems.IsEmpty() ? TSharedPtr<IChangelistTreeItem>() : SelectedChangelistItems[0];
-	OutStates.bShelvedFilesNodeSelected = false;
-	if (OutStates.SelectedChangelistNode && OutStates.SelectedChangelistNode->GetTreeItemType() == IChangelistTreeItem::ShelvedChangelist)
-	{
-		OutStates.SelectedChangelistNode = OutStates.SelectedChangelistNode->GetParent();
-		OutStates.bShelvedFilesNodeSelected = true;
-	}
-
-	// Save the selected item from 'uncontrolled changelists' tree.
-	SelectedChangelistItems = UncontrolledChangelistTreeView->GetSelectedItems();
-	OutStates.SelectedUncontrolledChangelistNode = SelectedChangelistItems.IsEmpty() ? TSharedPtr<IChangelistTreeItem>() : SelectedChangelistItems[0];
-
-	// Remember the expanded nodes.
-	check(OutStates.ExpandedTreeNodes.IsEmpty());
-	ChangelistTreeView->GetExpandedItems(OutStates.ExpandedTreeNodes);
-	UncontrolledChangelistTreeView->GetExpandedItems(OutStates.ExpandedTreeNodes);
-
-	// Remember the selected files.
-	OutStates.SelectedFileNodes.Reset(FileTreeView->GetNumItemsSelected());
-	OutStates.SelectedFileNodes.Append(FileTreeView->GetSelectedItems());
-}
-
-void SSourceControlChangelistsWidget::RestoreExpandedAndSelectionStates(const FExpandedAndSelectionStates& InStates)
-{
-	// Returns whether two changelist nodes represent the same changelist.
-	auto ChangelistEquals = [](const TSharedPtr<IChangelistTreeItem>& Lhs, const TSharedPtr<IChangelistTreeItem>& Rhs)
-	{
-		// NOTE: This TRUSTS the source control to return the same ' state' pointer before and after an update if the changelist still exists.
-		return static_cast<FChangelistTreeItem*>(Lhs.Get())->ChangelistState == static_cast<FChangelistTreeItem*>(Rhs.Get())->ChangelistState;
-	};
-
-	// Returns whether two uncontrolled changelist nodes represent the same changelist.
-	auto UncontrolledChangelistEquals = [](const TSharedPtr<IChangelistTreeItem>& Lhs, const TSharedPtr<IChangelistTreeItem>& Rhs)
-	{
-		// NOTE: This TRUSTS the source control to return the same 'state' pointer before and after an update if the changelist still exists.
-		return static_cast<FUncontrolledChangelistTreeItem*>(Lhs.Get())->UncontrolledChangelistState == static_cast<FUncontrolledChangelistTreeItem*>(Rhs.Get())->UncontrolledChangelistState;
-	};
-
-	// Find a specified item in a list. The nodes were deleted and recreated during the update and this function is used to match the new node corresponding to the old node.
-	auto Find = [](const TArray<TSharedPtr<IChangelistTreeItem>>& Nodes, const TSharedPtr<IChangelistTreeItem> SearchedItem,
-					const TFunction<bool(const TSharedPtr<IChangelistTreeItem>& Lhs, const TSharedPtr<IChangelistTreeItem>& Rhs)>& Predicate)
-	{
-		if (const TSharedPtr<IChangelistTreeItem>* Node = Nodes.FindByPredicate(
-			[&SearchedItem, &Predicate](const TSharedPtr<IChangelistTreeItem>& Candidate) { return Predicate(SearchedItem, Candidate); }))
-		{
-			return *Node;
-		}
-		return TSharedPtr<IChangelistTreeItem>(); // return nullptr;
-	};
-
-	// Restore the expansion states (Tree is only one level deep)
-	for (const TSharedPtr<IChangelistTreeItem>& ExpandedNode : InStates.ExpandedTreeNodes)
-	{
-		if (ExpandedNode->GetTreeItemType() == IChangelistTreeItem::Changelist)
-		{
-			// Check if the node still exist after the update.
-			if (TSharedPtr<IChangelistTreeItem> MatchingNode = Find(ChangelistTreeNodes, ExpandedNode, ChangelistEquals))
-			{
-				ChangelistTreeView->SetItemExpansion(MatchingNode, true);
-			}
-		}
-		else if (ExpandedNode->GetTreeItemType() == IChangelistTreeItem::UncontrolledChangelist)
-		{
-			// Check if the node still exist after the update.
-			if (TSharedPtr<IChangelistTreeItem> MatchingNode = Find(UncontrolledChangelistTreeNodes, ExpandedNode, UncontrolledChangelistEquals))
-			{
-				UncontrolledChangelistTreeView->SetItemExpansion(MatchingNode, true);
-			}
-		}
-	}
-
-	// Restore the selected nodes.
-	if (InStates.SelectedChangelistNode)
-	{
-		if (TSharedPtr<IChangelistTreeItem> MatchingNode = Find(ChangelistTreeNodes, InStates.SelectedChangelistNode, ChangelistEquals))
-		{
-			if (InStates.bShelvedFilesNodeSelected && static_cast<const FChangelistTreeItem*>(MatchingNode.Get())->GetShelvedFileCount() > 0)
-			{
-				for (const TSharedPtr<IChangelistTreeItem>& Child : MatchingNode->GetChildren())
-				{
-					if (Child->GetTreeItemType() == IChangelistTreeItem::ShelvedChangelist)
-					{
-						ChangelistTreeView->SetSelection(Child); // Select 'Shelved Files' node under the changelist.
-						break;
-					}
-				}
-			}
-			else
-			{
-				ChangelistTreeView->SetSelection(MatchingNode); // Select the 'changelist' node
-			}
-		}
-	}
-	else if (InStates.SelectedUncontrolledChangelistNode)
-	{
-		if (TSharedPtr<IChangelistTreeItem> MatchingNode = Find(UncontrolledChangelistTreeNodes, InStates.SelectedUncontrolledChangelistNode, UncontrolledChangelistEquals))
-		{
-			UncontrolledChangelistTreeView->SetSelection(MatchingNode); // Select the 'uncontrolled changelist' node
-		}
-	}
-
-	FileTreeView->ClearSelection();
-
-	// Try to reselect the files.
-	for (const TSharedPtr<IChangelistTreeItem>& FileNode : FileTreeNodes)
-	{
-		switch (FileNode->GetTreeItemType())
-		{
-			case IChangelistTreeItem::File:
-				if (InStates.SelectedFileNodes.ContainsByPredicate(
-						[&FileNode](const TSharedPtr<IChangelistTreeItem>& Candidate)
-						{
-							return Candidate->GetTreeItemType() == IChangelistTreeItem::File &&
-								static_cast<const FFileTreeItem*>(Candidate.Get())->GetAssetPath().EqualTo(static_cast<const FFileTreeItem*>(FileNode.Get())->GetAssetPath()) &&
-								static_cast<const FFileTreeItem*>(Candidate.Get())->GetFileName().EqualTo(static_cast<const FFileTreeItem*>(FileNode.Get())->GetFileName());
-						}))
-				{
-					FileTreeView->SetItemSelection(FileNode, true);
-				}
-				break;
-
-			case IChangelistTreeItem::ShelvedFile:
-				if (InStates.SelectedFileNodes.ContainsByPredicate(
-						[&FileNode](const TSharedPtr<IChangelistTreeItem>& Candidate)
-						{
-							return Candidate->GetTreeItemType() == IChangelistTreeItem::ShelvedFile &&
-								static_cast<const FShelvedFileTreeItem*>(Candidate.Get())->GetAssetPath().EqualTo(static_cast<const FShelvedFileTreeItem*>(FileNode.Get())->GetAssetPath()) &&
-								static_cast<const FShelvedFileTreeItem*>(Candidate.Get())->GetFileName().EqualTo(static_cast<const FShelvedFileTreeItem*>(FileNode.Get())->GetFileName());
-						}))
-				{
-					FileTreeView->SetItemSelection(FileNode, true);
-				}
-				break;
-
-			case IChangelistTreeItem::OfflineFile:
-				if (InStates.SelectedFileNodes.ContainsByPredicate(
-						[&FileNode](const TSharedPtr<IChangelistTreeItem>& Candidate)
-						{
-							return Candidate->GetTreeItemType() == IChangelistTreeItem::OfflineFile &&
-								static_cast<const FOfflineFileTreeItem*>(Candidate.Get())->GetDisplayPath().EqualTo(static_cast<const FOfflineFileTreeItem*>(FileNode.Get())->GetDisplayPath()) &&
-								static_cast<const FOfflineFileTreeItem*>(Candidate.Get())->GetDisplayName().EqualTo(static_cast<const FOfflineFileTreeItem*>(FileNode.Get())->GetDisplayName());
-						}))
-				{
-					FileTreeView->SetItemSelection(FileNode, true);
-				}
-				break;
-		}
-	}
+	OutdatedTimestampFiles.Emplace(PackageFilename); // The background task will convert relative to absolute path.
 }
 
 #undef LOCTEXT_NAMESPACE
diff --git a/Engine/Source/Editor/SourceControlWindows/Private/SSourceControlChangelists.h b/Engine/Source/Editor/SourceControlWindows/Private/SSourceControlChangelists.h
--- a/Engine/Source/Editor/SourceControlWindows/Private/SSourceControlChangelists.h
+++ b/Engine/Source/Editor/SourceControlWindows/Private/SSourceControlChangelists.h
@@ -4,25 +4,22 @@
 
 #include "CoreMinimal.h"
 
-#include "Widgets/SCompoundWidget.h"
-#include "Widgets/Views/STreeView.h"
-
+#include <atomic>
 #include "ISourceControlProvider.h"
-#include "SSourceControlCommon.h"
+#include "Misc/EnumClassFlags.h"
 #include "Misc/TextFilter.h"
+#include "SSourceControlCommon.h"
+#include "Stats/Stats.h"
+#include "UObject/ObjectSaveContext.h"
+#include "Widgets/SCompoundWidget.h"
+#include "Widgets/Views/SListView.h"
+#include "Widgets/Views/STreeView.h"
 
 class FChangelistGroupTreeItem;
 class SExpandableChangelistArea;
 class SSearchBox;
 class USourceControlSettings;
 
-class SChangelistTree : public STreeView<FChangelistTreeItemPtr>
-{
-private:
-	virtual void Private_SetItemSelection(FChangelistTreeItemPtr TheItem, bool bShouldBeSelected, bool bWasUserDirected = false) override;
-};
-
-
 /**
  * Displays the user source control change lists.
  */
@@ -34,34 +31,49 @@ public:
 
 	/** Constructs the widget */
 	void Construct(const FArguments& InArgs);
+	virtual ~SSourceControlChangelistsWidget();
 
 	/** Set selected files */
 	void SetSelectedFiles(const TArray<FString>& Filenames);
 
 private:
-	// Holds the list/state of selected and expanded items in the changelist views and file views.
-	struct FExpandedAndSelectionStates
+	enum class ERefreshFlags
 	{
-		TSharedPtr<IChangelistTreeItem> SelectedChangelistNode;
-		TSharedPtr<IChangelistTreeItem> SelectedUncontrolledChangelistNode;
-		TArray<TSharedPtr<IChangelistTreeItem>> SelectedFileNodes;
-		TSet<TSharedPtr<IChangelistTreeItem>> ExpandedTreeNodes;
-		bool bShelvedFilesNodeSelected = false;
+		SourceControlChangelists = 1 << 0,
+		UncontrolledChangelists  = 1 << 2,
+		All = SourceControlChangelists | UncontrolledChangelists,
+	};
+	FRIEND_ENUM_CLASS_FLAGS(ERefreshFlags)
+
+	/** Queries files timestamp asynchronously. */
+	class FAsyncTimestampUpdater
+	{
+	public:
+		void DoWork();
+		bool CanAbandon() { return true; }
+		void Abandon() { bAbandon = true; }
+		TStatId GetStatId() const { return TStatId(); }
+		static const TCHAR* Name() { return TEXT("SourceControlChangelistsTimestampUpdater"); }
+	public:
+		std::atomic<bool> bAbandon;
+		TSet<FString> RequestedFileTimestamps;
+		TMap<FString, FDateTime> QueriedFileTimestamps;
 	};
 
 private:
 	TSharedRef<SWidget> MakeToolBar();
-	TSharedRef<SChangelistTree> CreateChangelistTreeView(TArray<TSharedPtr<IChangelistTreeItem>>& ItemSources);
-	TSharedRef<STreeView<FChangelistTreeItemPtr>> CreateChangelistFilesView();
+	TSharedRef<STreeView<FChangelistTreeItemPtr>> CreateChangelistTreeView(TArray<TSharedPtr<IChangelistTreeItem>>& ItemSources);
+	TSharedRef<SListView<FChangelistTreeItemPtr>> CreateChangelistFilesView();
 
 	TSharedRef<ITableRow> OnGenerateRow(FChangelistTreeItemPtr InTreeItem, const TSharedRef<STableViewBase>& OwnerTable);
-	void OnGetFileChildren(FChangelistTreeItemPtr InParent, TArray<FChangelistTreeItemPtr>& OutChildren);
 	void OnGetChangelistChildren(FChangelistTreeItemPtr InParent, TArray<FChangelistTreeItemPtr>& OutChildren);
 	void OnFileViewHiddenColumnsListChanged();
 
 	EColumnSortPriority::Type GetColumnSortPriority(const FName ColumnId) const;
 	EColumnSortMode::Type GetColumnSortMode(const FName ColumnId) const;
 	void OnColumnSortModeChanged(const EColumnSortPriority::Type SortPriority, const FName& ColumnId, const EColumnSortMode::Type InSortMode);
+	bool IsFileViewSortedByFileStatusIcon() const;
+	bool IsFileViewSortedByLastModifiedTimestamp() const;
 	void SortFileView();
 
 	void OnChangelistSearchTextChanged(const FText& InFilterText);
@@ -74,7 +86,9 @@ private:
 	void RequestChangelistsRefresh();
 	void RequestFileStatusRefresh(const IChangelistTreeItem& Changelist);
 	void RequestFileStatusRefresh(const TSet<FString>& Pathnames);
-	void OnRefresh();
+	void OnRefreshUI(ERefreshFlags RefreshFlags);
+	void OnRefreshSourceControlWidgets(int64 CurreUpdateNum, const TFunction<void(TSharedPtr<IFileViewTreeItem>&)>& AddItemToFileView);
+	void OnRefreshUncontrolledChangelistWidgets(int64 CurreUpdateNum, const TFunction<void(TSharedPtr<IFileViewTreeItem>&)>& AddItemToFileView);
 	void ClearChangelistsTree();
 
 	TSharedPtr<SWidget> OnOpenContextMenu();
@@ -152,6 +166,9 @@ private:
 	void OnDiffAgainstWorkspace();
 	bool CanDiffAgainstWorkspace();
 
+	/** Uncontrolled changelist module callback. */
+	void OnUncontrolledChangelistStateChanged();
+
 	/** Source control callbacks */
 	void OnSourceControlProviderChanged(ISourceControlProvider& OldProvider, ISourceControlProvider& NewProvider);
 	void OnSourceControlStateChanged();
@@ -171,36 +188,40 @@ private:
 
 	virtual void Tick(const FGeometry& AllottedGeometry, const double InCurrentTime, const float InDeltaTime) override;
 
-	void SaveExpandedAndSelectionStates(FExpandedAndSelectionStates& OutStates);
-	void RestoreExpandedAndSelectionStates(const FExpandedAndSelectionStates& InStates);
-
 	/** Executes an operation to updates the changelist description of the provided changelist with a new description. */
 	void EditChangelistDescription(const FText& InNewChangelistDescription, const FSourceControlChangelistStatePtr& InChangelistState);
 
+	/** Invoked when a package is saved, to refresh the last saved timestamp. */
+	void OnPackageSaved(const FString& PackageFilename, UPackage* Package, FObjectPostSaveContext ObjectSaveContext);
+
 private:
 	TSharedPtr<SExpandableChangelistArea> ChangelistExpandableArea;
-	TSharedPtr<SExpandableChangelistArea> UncontrolledChangelistExpandableArea;
-
-	/** Hold the nodes displayed by the changelist tree. */
+	TSharedPtr<STreeView<FChangelistTreeItemPtr>> ChangelistTreeView;
 	TArray<TSharedPtr<IChangelistTreeItem>> ChangelistTreeNodes;
+
+	TSharedPtr<SExpandableChangelistArea> UncontrolledChangelistExpandableArea;
+	TSharedPtr<STreeView<FChangelistTreeItemPtr>> UncontrolledChangelistTreeView;
 	TArray<TSharedPtr<IChangelistTreeItem>> UncontrolledChangelistTreeNodes;
 
-	/** Hold the nodes displayed by the file tree. */
-	TArray<TSharedPtr<IChangelistTreeItem>> FileTreeNodes;
+	TSharedPtr<SListView<FChangelistTreeItemPtr>> FileListView;
+	TArray<TSharedPtr<IChangelistTreeItem>> FileListNodes;
+	TArray<FName> FileViewHiddenColumnsList;
 
-	/** Display the changelists, uncontrolled changelists and shelved nodes. */
-	TSharedPtr<SChangelistTree> ChangelistTreeView;
-	TSharedPtr<SChangelistTree> UncontrolledChangelistTreeView;
+	TMap<TSharedPtr<void>, TSharedPtr<IChangelistTreeItem>> SourceControlItemCache;
+	TMap<TSharedPtr<void>, TSharedPtr<IChangelistTreeItem>> UncontrolledChangelistItemCache;
+	TMap<FString, TSharedPtr<IFileViewTreeItem>> OfflineFileItemCache;
 
-	/** Display the list of files associated to the selected changelist, uncontrolled changelist or shelved node. */
-	TSharedPtr<STreeView<FChangelistTreeItemPtr>> FileTreeView;
-	TArray<FName> FileViewHiddenColumnsList;
+	TUniquePtr<FAsyncTask<FAsyncTimestampUpdater>> TimestampUpdateTask;
+	TSet<FString> OutdatedTimestampFiles;
 
 	/** Source control state changed delegate handle */
 	FDelegateHandle SourceControlStateChangedDelegateHandle;
 
+	int64 UpdateRequestNum = 0;
+	bool bInitialRefreshDone = false;
 	bool bShouldRefresh = false;
 	bool bSourceControlAvailable = false;
+	bool bUpdateMonitoredFileStatusList = false;
 
 	/** Files to select after refresh */
 	TArray<FString> FilesToSelect;
diff --git a/Engine/Source/Editor/SourceControlWindows/Private/SSourceControlCommon.cpp b/Engine/Source/Editor/SourceControlWindows/Private/SSourceControlCommon.cpp
--- a/Engine/Source/Editor/SourceControlWindows/Private/SSourceControlCommon.cpp
+++ b/Engine/Source/Editor/SourceControlWindows/Private/SSourceControlCommon.cpp
@@ -20,7 +20,9 @@
 #include "Widgets/Layout/SBox.h"
 #include "Framework/Docking/TabManager.h"
 #include "Framework/Notifications/NotificationManager.h"
+#include "Logging/MessageLog.h"
 #include "Misc/ScopedSlowTask.h"
+#include "Modules/ModuleManager.h"
 #include "Editor.h"
 
 #define LOCTEXT_NAMESPACE "SourceControlChangelist"
@@ -51,6 +53,15 @@ void IChangelistTreeItem::RemoveChild(const TSharedRef<IChangelistTreeItem>& Chi
 	}
 }
 
+void IChangelistTreeItem::RemoveAllChildren()
+{
+	for (TSharedPtr<IChangelistTreeItem>& Child : Children)
+	{
+		Child->Parent = nullptr;
+	}
+	Children.Reset();
+}
+
 static FString RetrieveAssetName(const FAssetData& InAssetData)
 {
 	static const FName NAME_ActorLabel(TEXT("ActorLabel"));
@@ -281,7 +292,7 @@ FText FFileTreeItem::GetAssetName()
 
 FText FShelvedChangelistTreeItem::GetDisplayText() const
 {
-	return FText::Format(LOCTEXT("SourceControl_ShelvedFiles", "Shelved Items ({0})"), Children.Num());
+	return LOCTEXT("SourceControl_ShelvedFiles", "Shelved Items");
 }
 
 //////////////////////////////////////////////////////////////////////////
@@ -401,6 +412,8 @@ void DisplaySourceControlOperationNotification(const FText& Message, SNotificati
 		return;
 	}
 
+	FMessageLog("SourceControl").Message(CompletionState == SNotificationItem::ECompletionState::CS_Fail ? EMessageSeverity::Error : EMessageSeverity::Info, Message);
+
 	FNotificationInfo NotificationInfo(Message);
 	NotificationInfo.ExpireDuration = 6.0f;
 	NotificationInfo.Hyperlink = FSimpleDelegate::CreateLambda([]() { FGlobalTabmanager::Get()->TryInvokeTab(FName("OutputLog")); });
diff --git a/Engine/Source/Editor/SourceControlWindows/Private/SSourceControlCommon.h b/Engine/Source/Editor/SourceControlWindows/Private/SSourceControlCommon.h
--- a/Engine/Source/Editor/SourceControlWindows/Private/SSourceControlCommon.h
+++ b/Engine/Source/Editor/SourceControlWindows/Private/SSourceControlCommon.h
@@ -56,6 +56,16 @@ struct IChangelistTreeItem : TSharedFromThis<IChangelistTreeItem>
 	/** Remove a child from this item */
 	void RemoveChild(const TSharedRef<IChangelistTreeItem>& Child);
 
+	/** Remove all children from this item. */
+	void RemoveAllChildren();
+
+public:
+	/** A sequence number representing the last time the item was inspected by the widget owning this UI item. Detect when the underlying model object stopped to exist between two UI updates.*/
+	int64 VisitedUpdateNum = -1;
+
+	/** A sequence number representing the last time the item was displayed by the widget owning this item. Used to detect when an existing item started/stopped to be visible between two UI updates.*/
+	int64 DisplayedUpdateNum = -1;
+
 protected:
 	IChangelistTreeItem(TreeItemType InType) { Type = InType; }
 
@@ -92,6 +102,9 @@ struct IFileViewTreeItem : public IChangelistTreeItem
 	/** The values displayed in the 'User' column. */
 	virtual const FString& GetCheckedOutBy() const { return DefaultStrValue; }
 
+	/** Returns the full pathname of the files on the file sytem. */
+	virtual const FString& GetFullPathname() const { return DefaultStrValue; }
+
 	/** Set the last modified time timestamp. */
 	void SetLastModifiedDateTime(const FDateTime& Timestamp);
 
@@ -101,6 +114,10 @@ struct IFileViewTreeItem : public IChangelistTreeItem
 	/** The values displayed in the 'Last Modified' column as text. */
 	FText GetLastModifiedTimestamp() const { return LastModifiedTimestampText; }
 
+public:
+	/** Keep the icon sorting priority as it was the last time the item was displayed. Used to detect if the priority changed between two refreshes of the UI. */
+	int32 DisplayedIconPriority = -1;
+
 protected:
 	IFileViewTreeItem(TreeItemType InType) : IChangelistTreeItem(InType) {}
 
@@ -115,6 +132,14 @@ private:
 };
 
 
+/** Root node to group shelved files as children. */
+struct FShelvedChangelistTreeItem : public IChangelistTreeItem
+{
+	FShelvedChangelistTreeItem() : IChangelistTreeItem(IChangelistTreeItem::ShelvedChangelist) {}
+	FText GetDisplayText() const;
+};
+
+
 /** Displays a changelist icon/number/description. */
 struct FChangelistTreeItem : public IChangelistTreeItem
 {
@@ -145,6 +170,7 @@ struct FChangelistTreeItem : public IChangelistTreeItem
 	}
 
 	TSharedRef<ISourceControlChangelistState> ChangelistState;
+	TSharedPtr<FShelvedChangelistTreeItem> ShelvedChangelistItem;
 };
 
 
@@ -185,6 +211,7 @@ struct FFileTreeItem : public IFileViewTreeItem
 	virtual const FString& GetName() const override { return AssetNameStr; }
 	virtual const FString& GetPath() const override { return AssetPathStr; }
 	virtual const FString& GetType() const override { return AssetTypeStr; }
+	virtual const FString& GetFullPathname() const override { return FileState->GetFilename(); }
 	virtual const FString& GetCheckedOutBy() const override;
 
 	/** Updates informations based on AssetData */
@@ -288,14 +315,6 @@ private:
 };
 
 
-/** Root node to group shelved files as children. */
-struct FShelvedChangelistTreeItem : public IChangelistTreeItem
-{
-	FShelvedChangelistTreeItem() : IChangelistTreeItem(IChangelistTreeItem::ShelvedChangelist) {}
-	FText GetDisplayText() const;
-};
-
-
 struct FShelvedFileTreeItem : public FFileTreeItem
 {
 	explicit FShelvedFileTreeItem(FSourceControlStateRef InFileState, bool bBeautifyPaths = true)
@@ -315,6 +334,7 @@ public:
 	virtual const FString& GetName() const override { return AssetNameStr; }
 	virtual const FString& GetPath() const override { return AssetPathStr; }
 	virtual const FString& GetType() const override { return AssetTypeStr; }
+	virtual const FString& GetFullPathname() const override { return Filename; }
 
 	const FString& GetFilename() const { return Filename; }
 	const FText& GetPackageName() const { return PackageName; }
@@ -376,4 +396,3 @@ typedef TSharedPtr<IChangelistTreeItem> FChangelistTreeItemPtr;
 typedef TSharedRef<IChangelistTreeItem> FChangelistTreeItemRef;
 typedef TSharedPtr<FFileTreeItem> FFileTreeItemPtr;
 typedef TSharedRef<FFileTreeItem> FFileTreeItemRef;
-
