Inverted the 'Save' and 'Submit' button position in the SSourceControlSubmit widget to be in the same position as Perforce submit dialog (from left to right now we have [Submit, Save, Cancel] instead of [Save, Submit, Cancel])

diff --git a/Engine/Source/Editor/SourceControlWindows/Private/SSourceControlSubmit.cpp b/Engine/Source/Editor/SourceControlWindows/Private/SSourceControlSubmit.cpp
--- a/Engine/Source/Editor/SourceControlWindows/Private/SSourceControlSubmit.cpp
+++ b/Engine/Source/Editor/SourceControlWindows/Private/SSourceControlSubmit.cpp
@@ -374,27 +374,20 @@ void SSourceControlSubmitWidget::Construct(const FArguments& InArgs)
 
 	const float AdditionalTopPadding = (bAllowKeepCheckedOut ? 0.0f : 5.0f);
 
+	TSharedPtr<SUniformGridPanel> SubmitSaveCancelButtonGrid;
+	int32 ButtonSlotId = 0;
+
 	Contents->AddSlot()
 	.AutoHeight()
 	.HAlign(HAlign_Right)
 	.VAlign(VAlign_Bottom)
-	.Padding(0.0f,AdditionalTopPadding,0.0f,5.0f)
+	.Padding(0.0f, AdditionalTopPadding, 0.0f, 5.0f)
 	[
-		SNew(SUniformGridPanel)
+		SAssignNew(SubmitSaveCancelButtonGrid, SUniformGridPanel)
 		.SlotPadding(FAppStyle::GetMargin("StandardDialog.SlotPadding"))
 		.MinDesiredSlotWidth(FAppStyle::GetFloat("StandardDialog.MinDesiredSlotWidth"))
 		.MinDesiredSlotHeight(FAppStyle::GetFloat("StandardDialog.MinDesiredSlotHeight"))
-		+SUniformGridPanel::Slot(0,0)
-		[
-			SNew(SButton)
-			.Visibility(bAllowSaveAndClose ? EVisibility::Visible : EVisibility::Collapsed)
-			.HAlign(HAlign_Center)
-			.ContentPadding(FAppStyle::GetMargin("StandardDialog.ContentPadding"))
-			.Text(NSLOCTEXT("SourceControl.SubmitPanel", "Save", "Save"))
-			.ToolTipText(NSLOCTEXT("SourceControl.SubmitPanel", "Save_Tooltip", "Save the description and close without submitting."))
-			.OnClicked(this, &SSourceControlSubmitWidget::SaveAndCloseClicked)
-		]
-		+SUniformGridPanel::Slot(1,0)
+		+SUniformGridPanel::Slot(ButtonSlotId++, 0)
 		[
 			SNew(SButton)
 			.HAlign(HAlign_Center)
@@ -403,15 +396,29 @@ void SSourceControlSubmitWidget::Construct(const FArguments& InArgs)
 			.Text( NSLOCTEXT("SourceControl.SubmitPanel", "OKButton", "Submit") )
 			.OnClicked(this, &SSourceControlSubmitWidget::SubmitClicked)
 		]
-		+SUniformGridPanel::Slot(2,0)
+	];
+
+	if (bAllowSaveAndClose)
+	{
+		SubmitSaveCancelButtonGrid->AddSlot(ButtonSlotId++, 0)
+			[
+				SNew(SButton)
+				.HAlign(HAlign_Center)
+				.ContentPadding(FAppStyle::GetMargin("StandardDialog.ContentPadding"))
+				.Text(NSLOCTEXT("SourceControl.SubmitPanel", "Save", "Save"))
+				.ToolTipText(NSLOCTEXT("SourceControl.SubmitPanel", "Save_Tooltip", "Save the description and close without submitting."))
+				.OnClicked(this, &SSourceControlSubmitWidget::SaveAndCloseClicked)
+			];
+	}
+
+	SubmitSaveCancelButtonGrid->AddSlot(ButtonSlotId++, 0)
 		[
 			SNew(SButton)
 			.HAlign(HAlign_Center)
 			.ContentPadding(FAppStyle::GetMargin("StandardDialog.ContentPadding"))
 			.Text( NSLOCTEXT("SourceControl.SubmitPanel", "CancelButton", "Cancel") )
 			.OnClicked(this, &SSourceControlSubmitWidget::CancelClicked)
-		]
-	];
+		];
 
 	RequestSort();
 
