SourceControl: Fixed status bar in 'View Changelists' dialog displaying 'Updating Changelists...' indefinitely when the provider does not support changelists (because the operation is unsupported, the callback is never called).

diff --git a/Engine/Source/Editor/SourceControlWindows/Private/SSourceControlChangelists.cpp b/Engine/Source/Editor/SourceControlWindows/Private/SSourceControlChangelists.cpp
--- a/Engine/Source/Editor/SourceControlWindows/Private/SSourceControlChangelists.cpp
+++ b/Engine/Source/Editor/SourceControlWindows/Private/SSourceControlChangelists.cpp
@@ -730,7 +730,7 @@ void SSourceControlChangelistsWidget::RequestChangelistsRefresh()
 {
 	bool bAnyProviderAvailable = false;
 
-	if (ISourceControlModule::Get().IsEnabled())
+	if (AreControlledChangelistsEnabled())
 	{
 		bAnyProviderAvailable = true;
 		StartRefreshStatus();
@@ -745,7 +745,7 @@ void SSourceControlChangelistsWidget::RequestChangelistsRefresh()
 		OnStartSourceControlOperation(UpdatePendingChangelistsOperation, LOCTEXT("SourceControl_UpdatingChangelist", "Updating changelists..."));
 	}
 
-	if (FUncontrolledChangelistsModule::Get().IsEnabled())
+	if (AreUncontrolledChangelistsEnabled())
 	{
 		bAnyProviderAvailable = true;
 
