SourceControl: For source control providers supporting the GetNumLocalChanges API, which returns a TOptional<bool>, the 'Submit Content' button in the Mainframe and SourceControl bar (bottom right) will only activate if there are changes to check-in.

diff --git a/Engine/Source/Editor/SourceControlWindows/Private/SourceControlWindows.cpp b/Engine/Source/Editor/SourceControlWindows/Private/SourceControlWindows.cpp
--- a/Engine/Source/Editor/SourceControlWindows/Private/SourceControlWindows.cpp
+++ b/Engine/Source/Editor/SourceControlWindows/Private/SourceControlWindows.cpp
@@ -142,10 +142,21 @@ bool FSourceControlWindows::CanChoosePackagesToCheckIn()
 {
 	ISourceControlModule& SourceControlModule = ISourceControlModule::Get();
 	
-	return ISourceControlModule::Get().IsEnabled() &&
+	if (ISourceControlModule::Get().IsEnabled() &&
 		ISourceControlModule::Get().GetProvider().IsAvailable() &&
-		!ChoosePackagesToCheckInNotification.IsValid()
-		;
+		!ChoosePackagesToCheckInNotification.IsValid())
+	{
+		if (SourceControlModule.GetProvider().GetNumLocalChanges().IsSet())
+		{
+			return SourceControlModule.GetProvider().GetNumLocalChanges().GetValue() > 0;
+		}
+		else
+		{
+			return true;
+		}
+	}
+ 
+	return false;
 }
 
 bool FSourceControlWindows::ShouldChoosePackagesToCheckBeVisible()
