Fixed an issue in the history window where strings that began with whitespace could show as empty in the description column

Standardized changelist description single line to a helper function to ensure the changelist window and history window have consistency

diff --git a/Engine/Source/Editor/SourceControlWindows/Private/SSourceControlChangelistRows.cpp b/Engine/Source/Editor/SourceControlWindows/Private/SSourceControlChangelistRows.cpp
--- a/Engine/Source/Editor/SourceControlWindows/Private/SSourceControlChangelistRows.cpp
+++ b/Engine/Source/Editor/SourceControlWindows/Private/SSourceControlChangelistRows.cpp
@@ -90,7 +90,7 @@ void SChangelistTableRow::Construct(const FArguments& InArgs, const TSharedRef<S
 			.AutoWidth()
 			[
 				SNew(STextBlock)
-				.Text(this, &SChangelistTableRow::GetChangelistDescriptionText)
+				.Text(this, &SChangelistTableRow::GetChangelistDescriptionSingleLineText)
 				.HighlightText(InArgs._HighlightText)
 			]
 		], InOwner);
@@ -99,7 +99,7 @@ void SChangelistTableRow::Construct(const FArguments& InArgs, const TSharedRef<S
 void SChangelistTableRow::PopulateSearchString(const FChangelistTreeItem& Item, TArray<FString>& OutStrings)
 {
 	OutStrings.Emplace(Item.GetDisplayText().ToString()); // The changelist number
-	OutStrings.Emplace(GetChangelistDescription(Item));   // The changelist description.
+	OutStrings.Emplace(Item.GetDescriptionText().ToString());   // The changelist description.
 }
 
 FText SChangelistTableRow::GetChangelistText() const
@@ -109,17 +109,12 @@ FText SChangelistTableRow::GetChangelistText() const
 
 FText SChangelistTableRow::GetChangelistDescriptionText() const
 {
-	return FText::FromString(GetChangelistDescription(*TreeItem));
+	return TreeItem->GetDescriptionText();
 }
 
-FString SChangelistTableRow::GetChangelistDescription(const FChangelistTreeItem& Item)
+FText SChangelistTableRow::GetChangelistDescriptionSingleLineText() const
 {
-	FString DescriptionString = Item.GetDescriptionText().ToString();
-	// Here we'll both remove \r\n (when edited from the dialog) and \n (when we get it from the SCC)
-	DescriptionString.ReplaceInline(TEXT("\r"), TEXT(""));
-	DescriptionString.ReplaceInline(TEXT("\n"), TEXT(" "));
-	DescriptionString.TrimEndInline();
-	return DescriptionString;
+	return SSourceControlCommon::GetSingleLineChangelistDescription(TreeItem->GetDescriptionText());
 }
 
 FReply SChangelistTableRow::OnDrop(const FGeometry& InGeometry, const FDragDropEvent& InDragDropEvent)
diff --git a/Engine/Source/Editor/SourceControlWindows/Private/SSourceControlChangelistRows.h b/Engine/Source/Editor/SourceControlWindows/Private/SSourceControlChangelistRows.h
--- a/Engine/Source/Editor/SourceControlWindows/Private/SSourceControlChangelistRows.h
+++ b/Engine/Source/Editor/SourceControlWindows/Private/SSourceControlChangelistRows.h
@@ -88,13 +88,13 @@ public:
 
 	FText GetChangelistText() const;
 	FText GetChangelistDescriptionText() const;
+	FText GetChangelistDescriptionSingleLineText() const;
 
 	static void PopulateSearchString(const FChangelistTreeItem& Item, TArray<FString>& OutStrings);
 
 protected:
 	//~ Begin STableRow Interface.
 	virtual FReply OnDrop(const FGeometry& InGeometry, const FDragDropEvent& InDragDropEvent) override;
-	static FString GetChangelistDescription(const FChangelistTreeItem& Item);
 
 private:
 	/** The info about the widget that we are visualizing. */
diff --git a/Engine/Source/Editor/SourceControlWindows/Private/SSourceControlCommon.cpp b/Engine/Source/Editor/SourceControlWindows/Private/SSourceControlCommon.cpp
--- a/Engine/Source/Editor/SourceControlWindows/Private/SSourceControlCommon.cpp
+++ b/Engine/Source/Editor/SourceControlWindows/Private/SSourceControlCommon.cpp
@@ -379,6 +379,21 @@ FText GetDefaultMultipleAsset()
 	return LOCTEXT("SourceCOntrol_ManyAssetType", "Multiple Assets");
 }
 
+FText GetSingleLineChangelistDescription(const FText& InFullDescription)
+{
+	FString DescriptionTextAsString = InFullDescription.ToString();
+	DescriptionTextAsString.TrimStartInline();
+
+	int32 NewlineStartIndex = INDEX_NONE;
+	DescriptionTextAsString.FindChar(TCHAR('\n'), NewlineStartIndex);
+	if (NewlineStartIndex != INDEX_NONE)
+	{
+		DescriptionTextAsString.LeftInline(NewlineStartIndex);
+	}
+
+	return InFullDescription.IsCultureInvariant() ? FText::AsCultureInvariant(DescriptionTextAsString) : FText::FromString(DescriptionTextAsString);
+}
+
 /** Wraps the execution of a changelist operations with a slow task. */
 void ExecuteChangelistOperationWithSlowTaskWrapper(const FText& Message, const TFunction<void()>& ChangelistTask)
 {
diff --git a/Engine/Source/Editor/SourceControlWindows/Private/SSourceControlCommon.h b/Engine/Source/Editor/SourceControlWindows/Private/SSourceControlCommon.h
--- a/Engine/Source/Editor/SourceControlWindows/Private/SSourceControlCommon.h
+++ b/Engine/Source/Editor/SourceControlWindows/Private/SSourceControlCommon.h
@@ -365,6 +365,11 @@ namespace SSourceControlCommon
 	FText GetDefaultUnknownAssetType();
 	FText GetDefaultMultipleAsset();
 
+	/**
+	 * returns the first non-whitespace line, or an empty FText if InFullDescription is empty or only whitespace
+	 */
+	FText GetSingleLineChangelistDescription(const FText& InFullDescription);
+
 	void ExecuteChangelistOperationWithSlowTaskWrapper(const FText& Message, const TFunction<void()>& ChangelistTask);
 	void ExecuteUncontrolledChangelistOperationWithSlowTaskWrapper(const FText& Message, const TFunction<void()>& UncontrolledChangelistTask);
 	void DisplaySourceControlOperationNotification(const FText& Message, SNotificationItem::ECompletionState CompletionState);
diff --git a/Engine/Source/Editor/SourceControlWindows/Private/SSourceControlHistory.cpp b/Engine/Source/Editor/SourceControlWindows/Private/SSourceControlHistory.cpp
--- a/Engine/Source/Editor/SourceControlWindows/Private/SSourceControlHistory.cpp
+++ b/Engine/Source/Editor/SourceControlWindows/Private/SSourceControlHistory.cpp
@@ -9,6 +9,7 @@
 #include "SourceControlWindows.h"
 #include "SourceControlHelpers.h"
 #include "ISourceControlModule.h"
+#include "SSourceControlCommon.h"
 #include "Modules/ModuleManager.h"
 #include "UObject/Object.h"
 #include "UObject/Package.h"
@@ -602,25 +603,10 @@ public:
 		}
 		else if (ColumnName == TEXT("Description"))
 		{
-			// Cut down the description to a single line for the list view
-			FString SingleLineDescription = RevisionListItem->Description;
-			int32 NewLinePos;
-			if (SingleLineDescription.FindChar(TCHAR('\n'), NewLinePos))
-			{
-				SingleLineDescription.LeftInline(NewLinePos, false);
-			}
-
-			// Trim any trailing new-line characters from the description for the tooltip
-			FString TooltipDescription = RevisionListItem->Description;
-			while(TooltipDescription.Len() && FChar::IsLinebreak(TooltipDescription[TooltipDescription.Len() - 1]))
-			{
-				TooltipDescription.RemoveAt(TooltipDescription.Len() - 1, 1, false);
-			}
-
 			return
 				SNew(STextBlock)
-				.Text(FText::FromString(SingleLineDescription))
-				.ToolTipText(FText::FromString(TooltipDescription));
+				.Text(SSourceControlCommon::GetSingleLineChangelistDescription(FText::FromString(RevisionListItem->Description)))
+				.ToolTipText(FText::FromString(RevisionListItem->Description));
 		}
 		else
 		{
