PR #9868: Revision Control History display file size with best unit (Contributed by [REDACTED_PERSON])

diff --git a/Engine/Source/Editor/SourceControlWindows/Private/SSourceControlHistory.cpp b/Engine/Source/Editor/SourceControlWindows/Private/SSourceControlHistory.cpp
--- a/Engine/Source/Editor/SourceControlWindows/Private/SSourceControlHistory.cpp
+++ b/Engine/Source/Editor/SourceControlWindows/Private/SSourceControlHistory.cpp
@@ -993,10 +993,7 @@ private:
 			static const FNumberFormattingOptions FileSizeFormatOptions = FNumberFormattingOptions()
 				.SetMinimumFractionalDigits(1)
 				.SetMaximumFractionalDigits(1);
-			return FText::Format(
-				NSLOCTEXT("SourceControlHistory", "FileSizeInMBFmt", "{0} MB"), 
-				FText::AsNumber(((float)LastSelectedRevisionItem.Pin()->FileSize) / (1024.f * 1024.f), &FileSizeFormatOptions)
-				);
+			return FText::AsMemory(LastSelectedRevisionItem.Pin()->FileSize, &FileSizeFormatOptions);
 		}
 		return FText::GetEmpty();
 	}
