[BugFix] Diffing Assets against depot can cause them to get saved with the PKG_ForDiffing flag
- non-temp assets shouldn't be loaded with the LOAD_ForDiff flag.

diff --git a/Engine/Source/Editor/SourceControlWindows/Private/SSourceControlHistory.cpp b/Engine/Source/Editor/SourceControlWindows/Private/SSourceControlHistory.cpp
--- a/Engine/Source/Editor/SourceControlWindows/Private/SSourceControlHistory.cpp
+++ b/Engine/Source/Editor/SourceControlWindows/Private/SSourceControlHistory.cpp
@@ -1666,7 +1666,7 @@ bool FSourceControlWindows::DiffAgainstWorkspace(const FString& InFileName)
 
 	if (FPackageName::TryConvertFilenameToLongPackageName(InFileName, AssetPackageName))
 	{
-		UPackage* AssetPackage = LoadPackage(nullptr, *AssetPackageName, LOAD_ForDiff | LOAD_DisableCompileOnLoad);
+		UPackage* AssetPackage = LoadPackage(nullptr, *AssetPackageName, LOAD_None);
 
 		// grab the asset from the package - we assume asset name matches file name
 		FString AssetName = FPaths::GetBaseFilename(InFileName);
@@ -1701,7 +1701,7 @@ bool FSourceControlWindows::DiffAgainstShelvedFile(const FSourceControlStateRef&
 	FString AssetPackageName;
 	if (FPackageName::TryConvertFilenameToLongPackageName(InFileState->GetFilename(), AssetPackageName))
 	{
-		UPackage* AssetPackage = LoadPackage(nullptr, *AssetPackageName, LOAD_ForDiff | LOAD_DisableCompileOnLoad);
+		UPackage* AssetPackage = LoadPackage(nullptr, *AssetPackageName, LOAD_None);
 
 		// grab the asset from the package - we assume asset name matches file name
 		FString AssetName = FPaths::GetBaseFilename(InFileState->GetFilename());
diff --git a/Engine/Source/Runtime/CoreUObject/Private/UObject/UObjectGlobals.cpp b/Engine/Source/Runtime/CoreUObject/Private/UObject/UObjectGlobals.cpp
--- a/Engine/Source/Runtime/CoreUObject/Private/UObject/UObjectGlobals.cpp
+++ b/Engine/Source/Runtime/CoreUObject/Private/UObject/UObjectGlobals.cpp
@@ -1936,6 +1936,13 @@ UPackage* LoadPackage(UPackage* InOuter, const FPackagePath& PackagePath, uint32
 											NAME_None,
 										 ELLMTagSet::Assets, ELLMTracker::Default);
 	TRACE_LOADTIME_REQUEST_GROUP_SCOPE(TEXT("SyncLoad - %s"), *PackagePath.GetDebugName());
+	
+	// if this is a supported asset, it should be loaded fully rather than just for diffing
+	if(!ensure(!(LoadFlags & LOAD_ForDiff) || FPackageName::IsTempPackage(PackagePath.GetPackageName())))
+	{
+		// clear LOAD_ForDiff
+		LoadFlags &= ~LOAD_ForDiff;
+	}
 	return LoadPackageInternal(InOuter, PackagePath, LoadFlags, /*ImportLinker =*/ nullptr, InReaderOverride, InstancingContext, DiffPackagePath);
 }
 
