Submit Files: Fix missing revision control icons

diff --git a/Engine/Source/Editor/SourceControlWindows/Private/SSourceControlSubmit.cpp b/Engine/Source/Editor/SourceControlWindows/Private/SSourceControlSubmit.cpp
--- a/Engine/Source/Editor/SourceControlWindows/Private/SSourceControlSubmit.cpp
+++ b/Engine/Source/Editor/SourceControlWindows/Private/SSourceControlSubmit.cpp
@@ -26,6 +26,7 @@
 #include "AssetRegistry/AssetRegistryModule.h"
 #include "Virtualization/VirtualizationSystem.h"
 #include "Logging/MessageLog.h"
+#include "RevisionControlStyle/RevisionControlStyle.h"
 
 #if SOURCE_CONTROL_WITH_SLATE
 
@@ -543,7 +544,7 @@ TSharedRef<SWidget> SSourceControlSubmitWidget::GenerateWidgetForItemAndColumn(T
 			.VAlign(VAlign_Center)
 			[
 				SNew(SImage)
-				.Image(FAppStyle::GetBrush(Item->GetIconName()))
+				.Image(FRevisionControlStyleManager::Get().GetBrush(Item->GetIconName()))
 				.ToolTipText(Item->GetIconTooltip())
 			];
 	}
