In the View Changelist window, hide the "Assets to check for reconcile" text if 'uncontrolled' changelists are not supported.

diff --git a/Engine/Source/Editor/SourceControlWindows/Private/SSourceControlChangelists.cpp b/Engine/Source/Editor/SourceControlWindows/Private/SSourceControlChangelists.cpp
--- a/Engine/Source/Editor/SourceControlWindows/Private/SSourceControlChangelists.cpp
+++ b/Engine/Source/Editor/SourceControlWindows/Private/SSourceControlChangelists.cpp
@@ -516,7 +516,7 @@ void SSourceControlChangelistsWidget::Construct(const FArguments& InArgs)
 				[
 					SNew(STextBlock)
 					.Text_Lambda([]() { return FUncontrolledChangelistsModule::Get().GetReconcileStatus(); })
-					.Visibility_Lambda([]() { return FUncontrolledChangelistsModule::Get().IsEnabled() ? EVisibility::Visible : EVisibility::Collapsed; })
+					.Visibility_Lambda([]() { return AreUncontrolledChangelistsEnabled() ? EVisibility::Visible : EVisibility::Collapsed; })
 				]
 			]
 		]
