Disable menu option for view changelists window if changelists are not supported by the current provider

diff --git a/Engine/Source/Editor/SourceControlWindows/Private/SourceControlWindowsModule.cpp b/Engine/Source/Editor/SourceControlWindows/Private/SourceControlWindowsModule.cpp
--- a/Engine/Source/Editor/SourceControlWindows/Private/SourceControlWindowsModule.cpp
+++ b/Engine/Source/Editor/SourceControlWindows/Private/SourceControlWindowsModule.cpp
@@ -69,7 +69,7 @@ void FSourceControlWindowsModule::StartupModule()
 		.SetDisplayName(LOCTEXT("ChangelistsTabTitle", "View Changelists"))
 		.SetTooltipText(LOCTEXT("ChangelistsTabTooltip", "Opens a dialog displaying current changelists."))
 		.SetIcon(SourceControlIcon)
-		.SetMenuType(ETabSpawnerMenuType::Hidden);	
+		.SetMenuType(ETabSpawnerMenuType::Hidden);
 
 #if WITH_RELOAD
 	// This code attempts to relaunch the GameplayCueEditor tab when you reload this module
@@ -123,7 +123,7 @@ bool FSourceControlWindowsModule::CanShowChangelistsTab() const
 {
 	ISourceControlModule& SourceControlModule = ISourceControlModule::Get();
 
-	return (SourceControlModule.IsEnabled() && SourceControlModule.GetProvider().IsAvailable() && SourceControlModule.GetProvider().UsesChangelists()) || FUncontrolledChangelistsModule::Get().IsEnabled();
+	return (SourceControlModule.IsEnabled() && SourceControlModule.GetProvider().IsAvailable() && SourceControlModule.GetProvider().UsesChangelists()) || (FUncontrolledChangelistsModule::Get().IsEnabled() && ISourceControlModule::Get().GetProvider().UsesUncontrolledChangelists());
 }
 
 void FSourceControlWindowsModule::SelectFiles(const TArray<FString>& Filenames)
