Add the ellipsis at the end of the description when the changelist description is too long to be displayed correctly.

diff --git a/Engine/Source/Editor/SourceControlWindows/Private/SSourceControlHistory.cpp b/Engine/Source/Editor/SourceControlWindows/Private/SSourceControlHistory.cpp
--- a/Engine/Source/Editor/SourceControlWindows/Private/SSourceControlHistory.cpp
+++ b/Engine/Source/Editor/SourceControlWindows/Private/SSourceControlHistory.cpp
@@ -606,7 +606,8 @@ public:
 			return
 				SNew(STextBlock)
 				.Text(SSourceControlCommon::GetSingleLineChangelistDescription(FText::FromString(RevisionListItem->Description)))
-				.ToolTipText(FText::FromString(RevisionListItem->Description));
+				.ToolTipText(FText::FromString(RevisionListItem->Description))
+				.OverflowPolicy(ETextOverflowPolicy::Ellipsis);
 		}
 		else
 		{
