SourceControl: Changelist #22880108 added support to the Submit window with support for exposing the diff functionality yes/no, depending on the SourceControl provider used. This changelist does the same for the History window.

diff --git a/Engine/Source/Editor/SourceControlWindows/Private/SSourceControlHistory.cpp b/Engine/Source/Editor/SourceControlWindows/Private/SSourceControlHistory.cpp
--- a/Engine/Source/Editor/SourceControlWindows/Private/SSourceControlHistory.cpp
+++ b/Engine/Source/Editor/SourceControlWindows/Private/SSourceControlHistory.cpp
@@ -1150,41 +1150,48 @@ private:
 	/** Called to create a context menu when right-clicking on a history item */
 	TSharedPtr< SWidget > OnCreateContextMenu()
 	{
-		FMenuBuilder MenuBuilder( true, NULL );
-
-		MenuBuilder.AddMenuEntry( 
-			NSLOCTEXT("SourceControl.HistoryWindow.Menu", "DiffAgainstPrev", "Diff Against Previous Revision"), 
-			NSLOCTEXT("SourceControl.HistoryWindow.Menu", "DiffAgainstPrevTooltip", "See changes between this revision and the previous one."), 
-			FSlateIcon(), 
-			FUIAction(
-				FExecuteAction::CreateSP( this, &SSourceControlHistoryWidget::OnDiffAgainstPreviousRev ),
-				FCanExecuteAction::CreateSP(this, &SSourceControlHistoryWidget::CanDiffAgainstPreviousRev)
-				)
-		);
+		if (CanDiff())
+		{
+			FMenuBuilder MenuBuilder(true, NULL);
 
-		MenuBuilder.AddMenuEntry(
-			NSLOCTEXT("SourceControl.HistoryWindow.Menu", "DiffAgainstWorkspace", "Diff Against Workspace File"),
-			NSLOCTEXT("SourceControl.HistoryWindow.Menu", "DiffAgainstWorkspaceTooltip", "See changes between this revision and your version of the asset."),
-			FSlateIcon(),
-			FUIAction(
-				FExecuteAction::CreateSP(this, &SSourceControlHistoryWidget::OnDiffAgainstWorkspace),
-				FCanExecuteAction::CreateSP(this, &SSourceControlHistoryWidget::CanDiffAgainstWorkspace)
-			)
-		);
+			MenuBuilder.AddMenuEntry(
+				NSLOCTEXT("SourceControl.HistoryWindow.Menu", "DiffAgainstPrev", "Diff Against Previous Revision"),
+				NSLOCTEXT("SourceControl.HistoryWindow.Menu", "DiffAgainstPrevTooltip", "See changes between this revision and the previous one."),
+				FSlateIcon(),
+				FUIAction(
+					FExecuteAction::CreateSP(this, &SSourceControlHistoryWidget::OnDiffAgainstPreviousRev),
+					FCanExecuteAction::CreateSP(this, &SSourceControlHistoryWidget::CanDiffAgainstPreviousRev)
+				)
+			);
 
-		if (CanDiffSelected())
-		{
-			MenuBuilder.AddMenuEntry( 
-				NSLOCTEXT("SourceControl.HistoryWindow.Menu", "DiffSelected", "Diff Selected"), 
-				NSLOCTEXT("SourceControl.HistoryWindow.Menu", "DiffSelectedTooltip", "Diff the two assets that you have selected."), 
-				FSlateIcon(), 
+			MenuBuilder.AddMenuEntry(
+				NSLOCTEXT("SourceControl.HistoryWindow.Menu", "DiffAgainstWorkspace", "Diff Against Workspace File"),
+				NSLOCTEXT("SourceControl.HistoryWindow.Menu", "DiffAgainstWorkspaceTooltip", "See changes between this revision and your version of the asset."),
+				FSlateIcon(),
 				FUIAction(
-					FExecuteAction::CreateSP(this, &SSourceControlHistoryWidget::OnDiffSelected)
+					FExecuteAction::CreateSP(this, &SSourceControlHistoryWidget::OnDiffAgainstWorkspace),
+					FCanExecuteAction::CreateSP(this, &SSourceControlHistoryWidget::CanDiffAgainstWorkspace)
 				)
 			);
-		}
 
-		return MenuBuilder.MakeWidget();
+			if (CanDiffSelected())
+			{
+				MenuBuilder.AddMenuEntry(
+					NSLOCTEXT("SourceControl.HistoryWindow.Menu", "DiffSelected", "Diff Selected"),
+					NSLOCTEXT("SourceControl.HistoryWindow.Menu", "DiffSelectedTooltip", "Diff the two assets that you have selected."),
+					FSlateIcon(),
+					FUIAction(
+						FExecuteAction::CreateSP(this, &SSourceControlHistoryWidget::OnDiffSelected)
+					)
+				);
+			}
+
+			return MenuBuilder.MakeWidget();
+		}
+		else
+		{
+			return SNullWidget::NullWidget;
+		}
 	}
 
 	/** See if we should enabled the 'diff against previous' option */
@@ -1405,6 +1412,16 @@ private:
 		}
 	}
 
+	/**
+	 * Checks to see if the SourceControl provider supports diffing against depot.
+	 *
+	 * @return True if the SourceControl provider supports diffing, false if not.
+	 */
+	bool CanDiff() const
+	{
+		return ISourceControlModule::Get().GetProvider().AllowsDiffAgainstDepot();
+	}
+
 	/**
 	 * Checks to see if the selected history-tree items can be diffed against each other.
 	 * 
