For Source Control systems which don't support Uncontrolled Changelists, disable the entire Uncontrolled Changelists module, rather than just UI elements (as was done in CL 23490912).
This prevents uncontrolled icons being shown in the outliner when not supported by the current source control.

diff --git a/Engine/Source/Developer/UncontrolledChangelists/Private/UncontrolledChangelistsModule.cpp b/Engine/Source/Developer/UncontrolledChangelists/Private/UncontrolledChangelistsModule.cpp
--- a/Engine/Source/Developer/UncontrolledChangelists/Private/UncontrolledChangelistsModule.cpp
+++ b/Engine/Source/Developer/UncontrolledChangelists/Private/UncontrolledChangelistsModule.cpp
@@ -105,7 +105,7 @@ void FUncontrolledChangelistsModule::ShutdownModule()
 
 bool FUncontrolledChangelistsModule::IsEnabled() const
 {
-	return bIsEnabled;
+	return bIsEnabled && ISourceControlModule::Get().GetProvider().UsesUncontrolledChangelists();
 }
 
 TArray<FUncontrolledChangelistStateRef> FUncontrolledChangelistsModule::GetChangelistStates() const
diff --git a/Engine/Source/Editor/SourceControlWindows/Private/SSourceControlChangelists.cpp b/Engine/Source/Editor/SourceControlWindows/Private/SSourceControlChangelists.cpp
--- a/Engine/Source/Editor/SourceControlWindows/Private/SSourceControlChangelists.cpp
+++ b/Engine/Source/Editor/SourceControlWindows/Private/SSourceControlChangelists.cpp
@@ -62,7 +62,7 @@ bool AreControlledChangelistsEnabled()
 /** Returns true if Uncontrolled changelists are enabled. */
 bool AreUncontrolledChangelistsEnabled()
 {
-	return FUncontrolledChangelistsModule::Get().IsEnabled() && ISourceControlModule::Get().GetProvider().UsesUncontrolledChangelists();
+	return FUncontrolledChangelistsModule::Get().IsEnabled();
 };
 
 /** Returns true if there are changelists to display. */
