Unsaved Assets in Changelist Window

diff --git a/Engine/Source/Developer/SourceControl/Private/SourceControlHelpers.cpp b/Engine/Source/Developer/SourceControl/Private/SourceControlHelpers.cpp
--- a/Engine/Source/Developer/SourceControl/Private/SourceControlHelpers.cpp
+++ b/Engine/Source/Developer/SourceControl/Private/SourceControlHelpers.cpp
@@ -908,7 +908,7 @@ bool USourceControlHelpers::RevertFile(const FString& InFile, bool bSilent)
 }
 
 #if WITH_EDITOR
-bool USourceControlHelpers::ApplyOperationAndReloadPackages(const TArray<FString>& InFilenames,	const TFunctionRef<bool(const TArray<FString>&)>& InOperation, bool bReloadWorld)
+bool USourceControlHelpers::ApplyOperationAndReloadPackages(const TArray<FString>& InFilenames,	const TFunctionRef<bool(const TArray<FString>&)>& InOperation, bool bReloadWorld, bool bInteractive)
 {
 	ISourceControlProvider& SourceControlProvider = ISourceControlModule::Get().GetProvider();
 	TArray<UPackage*> LoadedPackages;
@@ -1037,8 +1037,7 @@ bool USourceControlHelpers::ApplyOperationAndReloadPackages(const TArray<FString
 
 	// Hot-reload the new packages...
 	FText OutReloadErrorMsg;
-	const bool bInteractive = true;
-	UPackageTools::ReloadPackages(LoadedPackages, OutReloadErrorMsg, EReloadPackagesInteractionMode::Interactive);
+	UPackageTools::ReloadPackages(LoadedPackages, OutReloadErrorMsg, bInteractive ? EReloadPackagesInteractionMode::Interactive : EReloadPackagesInteractionMode::AssumePositive);
 	if (!OutReloadErrorMsg.IsEmpty())
 	{
 		UE_LOG(LogSourceControl, Warning, TEXT("%s"), *OutReloadErrorMsg.ToString());
diff --git a/Engine/Source/Developer/SourceControl/Public/SourceControlHelpers.h b/Engine/Source/Developer/SourceControl/Public/SourceControlHelpers.h
--- a/Engine/Source/Developer/SourceControl/Public/SourceControlHelpers.h
+++ b/Engine/Source/Developer/SourceControl/Public/SourceControlHelpers.h
@@ -354,10 +354,11 @@ public:
 	 * @param	InPackagesToApplyOperation			The files/packages to apply the operation
 	 * @param	InOperation							The function to apply
 	 * @param	bReloadWorld						Reload the world
+	 * @param	bInteractive						Whether to prompt user when discarding changes
 	 * @return true if succeeded.
 	 */
 	static bool ApplyOperationAndReloadPackages(const TArray<FString>& InPackagesToApplyOperation, 
-		const TFunctionRef<bool(const TArray<FString>&)>& InOperation, bool bReloadWorld = false);
+		const TFunctionRef<bool(const TArray<FString>&)>& InOperation, bool bReloadWorld = false, bool bInteractive = true);
 
 	/**
      * Reverts the provided files then reloads packages.
diff --git a/Engine/Source/Editor/SourceControlWindows/Private/SSourceControlChangelistRows.cpp b/Engine/Source/Editor/SourceControlWindows/Private/SSourceControlChangelistRows.cpp
--- a/Engine/Source/Editor/SourceControlWindows/Private/SSourceControlChangelistRows.cpp
+++ b/Engine/Source/Editor/SourceControlWindows/Private/SSourceControlChangelistRows.cpp
@@ -3,12 +3,22 @@
 #include "SSourceControlChangelistRows.h"
 
 #include "ISourceControlModule.h"
+#include "PackageTools.h"
+#include "SourceControlHelpers.h"
 #include "UncontrolledChangelistsModule.h"
 #include "SourceControlOperations.h"
+#include "UnsavedAssetsTrackerModule.h"
 #include "Styling/AppStyle.h"
 #include "Widgets/Images/SImage.h"
+#include "Widgets/Input/SButton.h"
 #include "Widgets/Text/STextBlock.h"
 #include "Algo/Transform.h"
+#include "Async/Async.h"
+#include "RevisionControlStyle/RevisionControlStyle.h"
+#include "Styling/SlateStyleRegistry.h"
+#include "Styling/StarshipCoreStyle.h"
+#include "Widgets/Images/SThrobber.h"
+#include "Widgets/Layout/SWidgetSwitcher.h"
 
 #define LOCTEXT_NAMESPACE "SourceControlChangelistRow"
 
@@ -36,6 +46,18 @@ FName SourceControlFileViewColumn::CheckedOutByUser::Id() { return TEXT("Checked
 FText SourceControlFileViewColumn::CheckedOutByUser::GetDisplayText() { return LOCTEXT("CheckedOutByUser_Column", "User"); }
 FText SourceControlFileViewColumn::CheckedOutByUser::GetToolTipText() { return LOCTEXT("CheckedOutByUser_Column_Tooltip", "Displays the other user(s) that checked out the file/asset, if any"); }
 
+FName SourceControlFileViewColumn::Changelist::Id() { return TEXT("Changelist"); }
+FText SourceControlFileViewColumn::Changelist::GetDisplayText() { return LOCTEXT("Changelist_Column", "Changelist"); }
+FText SourceControlFileViewColumn::Changelist::GetToolTipText() { return LOCTEXT("Changelist_Column_Tooltip", "Displays the changelist the asset/file belongs to, if any"); }
+
+FName SourceControlFileViewColumn::Dirty::Id() { return TEXT("Dirty"); }
+FText SourceControlFileViewColumn::Dirty::GetDisplayText() { return LOCTEXT("Dirty_Column", "Unsaved"); }
+FText SourceControlFileViewColumn::Dirty::GetToolTipText() { return LOCTEXT("Dirty_Column_Tooltip", "Displays whether the asset/file has unsaved changes"); }
+
+FName SourceControlFileViewColumn::Discard::Id() { return TEXT("Discard"); }
+FText SourceControlFileViewColumn::Discard::GetDisplayText() { return LOCTEXT("Discard_Column", "Discard Unsaved Changes"); }
+FText SourceControlFileViewColumn::Discard::GetToolTipText() { return LOCTEXT("Discard_Column_Tooltip", "Provides option to discard unsaved changes to an asset/file"); }
+
 FText FormatChangelistFileCountText(int32 DisplayedCount, int32 TotalCount)
 {
 	return DisplayedCount == TotalCount ?
@@ -242,6 +264,39 @@ FReply SUncontrolledChangelistTableRow::OnDrop(const FGeometry& InGeometry, cons
 	return FReply::Handled();
 }
 
+void SUnsavedAssetsTableRow::Construct(const FArguments& InArgs, const TSharedRef<STableViewBase>& InOwner)
+{
+	STableRow<FChangelistTreeItemPtr>::Construct(
+		STableRow<FChangelistTreeItemPtr>::FArguments()
+		.Style(FAppStyle::Get(), "TableView.Row")
+		.Content()
+		[
+			SNew(SHorizontalBox)
+			+ SHorizontalBox::Slot()
+			.AutoWidth()
+			.VAlign(VAlign_Center)
+			[
+				SNew(SImage)
+				.Image(FAppStyle::GetBrush("Assets.Unsaved"))
+			]
+			+ SHorizontalBox::Slot()
+			.AutoWidth()
+			.Padding(2.0f, 0.0f)
+			.VAlign(VAlign_Center)
+			[
+				SNew(STextBlock)
+				.Text(LOCTEXT("SourceControl_Unsaved", "Unsaved"))
+			]
+			+SHorizontalBox::Slot() // Files/Offline file count.
+			.Padding(4, 0, 4, 0)
+			.AutoWidth()
+			[
+				SNew(STextBlock)
+				.Text_Lambda([] { return FText::Format(FText::FromString("({0})"), FUnsavedAssetsTrackerModule::Get().GetUnsavedAssetNum()); })
+			]
+		], InOwner);
+}
+
 
 void SFileTableRow::Construct(const FArguments& InArgs, const TSharedRef<STableViewBase>& InOwner)
 {
@@ -463,6 +518,74 @@ TSharedRef<SWidget> SOfflineFileTableRow::GenerateWidgetForColumn(const FName& C
 		return SNew(STextBlock)
 			.Text(FText::GetEmpty());
 	}
+	else if (ColumnId == SourceControlFileViewColumn::Dirty::Id())
+	{
+		if (FUnsavedAssetsTrackerModule::Get().IsAssetUnsaved(GetFilename().ToString()))
+		{
+			return SNew(SBox)
+				.WidthOverride(16) // Small Icons are usually 16x16
+				.HAlign(HAlign_Center)
+				[
+					SNew(SImage)
+					.Image(FAppStyle::GetBrush(FName("SourceControl.OfflineFile_Small")))
+				];
+		}
+		return SNew(SBox)
+			.WidthOverride(16); // Small Icons are usually 16x16
+	}
+	else if (ColumnId == SourceControlFileViewColumn::Discard::Id())
+	{
+		FString Filename = TreeItem->GetFullPathname();
+		if (!FUnsavedAssetsTrackerModule::Get().IsAssetUnsaved(Filename))
+		{
+			return SNew(SBox)
+				.WidthOverride(16); // Small Icons are usually 16x16
+		}
+
+		TSharedRef<SWidgetSwitcher> DiscardSwitcher = SNew(SWidgetSwitcher);
+
+		TSharedRef<SImage> DiscardButton = SNew(SImage)
+				.DesiredSizeOverride(FVector2D{ 16.0f })
+				.Image(FAppStyle::Get().GetBrush("Icons.XCircle"))
+				.ColorAndOpacity(FSlateColor::UseSubduedForeground())
+				.OnMouseButtonDown_Lambda([this, DiscardSwitcher] (const FGeometry&, const FPointerEvent&) -> FReply
+				{
+					DiscardSwitcher->SetActiveWidgetIndex(1);
+					Async(EAsyncExecution::TaskGraphMainThread,
+						[this]
+						{
+							TArray<FString> PackageToReload { TreeItem->GetPackageName().ToString() };
+							const bool bAllowReloadWorld = true;
+							const bool bInteractive = false;
+							USourceControlHelpers::ApplyOperationAndReloadPackages(
+								PackageToReload,
+								[](const TArray<FString>&) -> bool { return true; },
+								bAllowReloadWorld,
+								bInteractive
+							);
+						},
+						[] {  });
+					
+					return FReply::Handled();
+				});
+
+		DiscardSwitcher->AddSlot()[
+			DiscardButton
+		];
+		DiscardSwitcher->AddSlot()[
+			SNew(SCircularThrobber)
+			.Radius(7.5f)
+		];
+		
+		return
+			SNew(SBox)
+			.WidthOverride(16) // Small Icons are usually 16x16
+			.Padding(FMargin{1, 0})
+			.ToolTipText(LOCTEXT("UnsavedAsset_DiscardChanges", "Discard unsaved changes"))
+			[
+				DiscardSwitcher
+			];
+	}
 	else
 	{
 		return SNullWidget::NullWidget;
diff --git a/Engine/Source/Editor/SourceControlWindows/Private/SSourceControlChangelistRows.h b/Engine/Source/Editor/SourceControlWindows/Private/SSourceControlChangelistRows.h
--- a/Engine/Source/Editor/SourceControlWindows/Private/SSourceControlChangelistRows.h
+++ b/Engine/Source/Editor/SourceControlWindows/Private/SSourceControlChangelistRows.h
@@ -60,6 +60,30 @@ namespace SourceControlFileViewColumn
 		FText GetDisplayText();
 		FText GetToolTipText();
 	}
+
+	/** "What changelist a given file belongs to" column. */
+	namespace Changelist
+	{
+		FName Id();
+		FText GetDisplayText();
+		FText GetToolTipText();
+	}
+
+	/** "Whether a file is unsaved/dirty" column */
+	namespace Dirty
+	{
+		FName Id();
+		FText GetDisplayText();
+		FText GetToolTipText();
+	}
+
+	/** "Discard unsaved changes" column */
+	namespace Discard
+	{
+		FName Id();
+		FText GetDisplayText();
+		FText GetToolTipText();
+	}
 }
 
 
@@ -144,6 +168,17 @@ private:
 	FSimpleDelegate OnPostDrop;
 };
 
+/** Displays a block for unsaved assets with icon and count */
+class SUnsavedAssetsTableRow : public STableRow<FChangelistTreeItemPtr>
+{
+public:
+	/**
+	* Construct child widgets that comprise this widget.
+	*
+	* @param InArgs Declaration from which to construct this widget.
+	*/
+	void Construct(const FArguments& InArgs, const TSharedRef<STableViewBase>& InOwner);
+};
 
 /** Display information about a file (icon, name, location, type, etc.) */
 class SFileTableRow : public SMultiColumnTableRow<FChangelistTreeItemPtr>
diff --git a/Engine/Source/Editor/SourceControlWindows/Private/SSourceControlChangelists.cpp b/Engine/Source/Editor/SourceControlWindows/Private/SSourceControlChangelists.cpp
--- a/Engine/Source/Editor/SourceControlWindows/Private/SSourceControlChangelists.cpp
+++ b/Engine/Source/Editor/SourceControlWindows/Private/SSourceControlChangelists.cpp
@@ -16,6 +16,7 @@
 #include "Widgets/Layout/SSpacer.h"
 #include "Widgets/Layout/SSplitter.h"
 #include "Widgets/Layout/SExpandableArea.h"
+#include "Widgets/Layout/SWidgetSwitcher.h"
 
 #include "ISourceControlProvider.h"
 #include "ISourceControlModule.h"
@@ -31,14 +32,17 @@
 #include "SourceControlMenuContext.h"
 #include "SourceControlSettings.h"
 #include "AssetToolsModule.h"
+#include "FileHelpers.h"
 #include "Misc/CString.h"
 #include "Misc/MessageDialog.h"
 #include "Misc/PackageName.h"
 #include "Misc/ScopedSlowTask.h"
+#include "PackageSourceControlHelper.h"
 #include "Algo/AnyOf.h"
 #include "HAL/PlatformTime.h"
 #include "HAL/FileManager.h"
 #include "ToolMenus.h"
+#include "UnsavedAssetsTracker/Source/Public/UnsavedAssetsTrackerModule.h"
 #include "UObject/UObjectGlobals.h"
 #include "UObject/Package.h"
 #include "SSourceControlSubmit.h"
@@ -53,6 +57,11 @@
 namespace
 {
 
+bool AreUnsavedAssetsPresent()
+{
+	return FUnsavedAssetsTrackerModule::Get().GetUnsavedAssetNum() > 0;
+}
+	
 /** Returns true if a source control provider is enable and support changeslists. */
 bool AreControlledChangelistsEnabled()
 {
@@ -338,10 +347,24 @@ void SSourceControlChangelistsWidget::Construct(const FArguments& InArgs)
 	// No default sorting, sorting make the view slower, just pay the cost if user wants it.
 	PrimarySortedColumn = FName();
 
+	UnsavedAssetsTreeNode.Add(MakeShared<FUnsavedAssetsTreeItem>());
+	
+	UnsavedAssetsTreeView = CreateChangelistTreeView(UnsavedAssetsTreeNode);
 	ChangelistTreeView = CreateChangelistTreeView(ChangelistTreeNodes);
 	UncontrolledChangelistTreeView = CreateChangelistTreeView(UncontrolledChangelistTreeNodes);
 	FileListView = CreateChangelistFilesView();
-
+	UnsavedAssetsFileListView = CreateUnsavedAssetsFilesView();
+	
+	FileListViewSwitcher = SNew(SWidgetSwitcher);
+	FileListViewSwitcher->AddSlot().AttachWidget(FileListView.ToSharedRef());
+	FileListViewSwitcher->AddSlot().AttachWidget(UnsavedAssetsFileListView.ToSharedRef());
+
+	UnsavedAssetsExpandableArea = SNew(SExpandableChangelistArea)
+		.HeaderText_Lambda([] { return FText::Format(LOCTEXT("SourceControl_UnsavedAssets", "Unsaved Assets ({0})"), FUnsavedAssetsTrackerModule::Get().GetUnsavedAssetNum()); })
+		.ChangelistView(UnsavedAssetsTreeView.ToSharedRef())
+		.NewButtonVisibility(EVisibility::Collapsed)
+		.SearchButtonVisibility(EVisibility::Collapsed);
+	
 	ChangelistExpandableArea = SNew(SExpandableChangelistArea)
 		.HeaderText_Lambda([this]() { return FText::Format(LOCTEXT("SourceControl_ChangeLists", "Changelists ({0})"), ChangelistTreeNodes.Num()); })
 		.ChangelistView(ChangelistTreeView.ToSharedRef())
@@ -367,6 +390,9 @@ void SSourceControlChangelistsWidget::Construct(const FArguments& InArgs)
 	FileTextFilter = MakeShared<TTextFilter<const IChangelistTreeItem&>>(TTextFilter<const IChangelistTreeItem&>::FItemToStringArray::CreateSP(this, &SSourceControlChangelistsWidget::PopulateItemSearchStrings));
 	FileTextFilter->OnChanged().AddSP(this, &SSourceControlChangelistsWidget::OnRefreshUI, ERefreshFlags::All);
 
+	FUnsavedAssetsTrackerModule::Get().OnUnsavedAssetAdded.AddSP(this, &SSourceControlChangelistsWidget::OnUnsavedAssetChanged);
+	FUnsavedAssetsTrackerModule::Get().OnUnsavedAssetRemoved.AddSP(this, &SSourceControlChangelistsWidget::OnUnsavedAssetChanged);
+
 	// Min/Max prevents making the Changelist Area too small and consequently prevent the file search box to go over the 'refresh' button.
 	const float MinChangelistAreaRatio = 0.2f;
 	const float MaxFileAreaRation = 1.0f - MinChangelistAreaRatio;
@@ -417,39 +443,48 @@ void SSourceControlChangelistsWidget::Construct(const FArguments& InArgs)
 		]
 		+SVerticalBox::Slot() // Everything below the tools bar: changelist expandable areas + files views + status bar at the bottom
 		[
-			SNew(SOverlay)
-			+SOverlay::Slot()
+			SNew(SSplitter)
+			.Orientation(EOrientation::Orient_Horizontal)
+			.ResizeMode(ESplitterResizeMode::FixedPosition)
+
+			// Left slot: Changelists and uncontrolled changelists areas
+			+SSplitter::Slot()
+			.Resizable(true)
+			.SizeRule(SSplitter::FractionOfParent)
+			.Value_Lambda([this, MinChangelistAreaRatio]() { return FMath::Max(MinChangelistAreaRatio, ChangelistAreaSize); })
+			.OnSlotResized_Lambda([this](float Size) { ChangelistAreaSize = Size; })
 			[
-				SNew(SBox)
-				.Visibility_Lambda([](){ return !AreChangelistsEnabled() ? EVisibility::Visible: EVisibility::Collapsed; })
-				.HAlign(HAlign_Center)
-				.VAlign(VAlign_Center)
+				SNew(SVerticalBox)
+				+SVerticalBox::Slot()
+				.AutoHeight()
 				[
-					SNew(STextBlock)
-					.Text(LOCTEXT("SourceControl_Disabled", "The revision control is disabled or it doesn't support changelists."))
+					SNew(SBox)
+					.Visibility_Lambda([] { return AreUnsavedAssetsPresent() ? EVisibility::Visible : EVisibility::Collapsed; } )
+					[
+						UnsavedAssetsExpandableArea.ToSharedRef()
+					]
 				]
-			]
-			+SOverlay::Slot()
-			[
-				SNew(SSplitter)
-				.Orientation(EOrientation::Orient_Horizontal)
-				.ResizeMode(ESplitterResizeMode::FixedPosition)
-				.Visibility_Lambda([]() { return AreChangelistsEnabled() ? EVisibility::Visible : EVisibility::Collapsed; })
-
-				// Left slot: Changelists and uncontrolled changelists areas
-				+SSplitter::Slot()
-				.Resizable(true)
-				.SizeRule(SSplitter::FractionOfParent)
-				.Value_Lambda([this, MinChangelistAreaRatio]() { return FMath::Max(MinChangelistAreaRatio, ChangelistAreaSize); })
-				.OnSlotResized_Lambda([this](float Size) { ChangelistAreaSize = Size; })
+				+SVerticalBox::Slot()
+				.AutoHeight()
 				[
-					SNew(SOverlay) // Visible when both Controlled and Uncontrolled changelists are enabled (Need to split vertical space)
+					SNew(SOverlay)
 					+SOverlay::Slot()
+					[
+						SNew(SBox)
+						.Visibility_Lambda([](){ return !AreChangelistsEnabled() ? EVisibility::Visible: EVisibility::Collapsed; })
+						.HAlign(HAlign_Center)
+						.VAlign(VAlign_Center)
+						[
+							SNew(STextBlock)
+							.Text(LOCTEXT("SourceControl_Disabled", "The revision control is disabled or it doesn't support changelists."))
+						]
+					]
+					+SOverlay::Slot() // Visible when both Controlled and Uncontrolled changelists are enabled (Need to split vertical space)
 					[
 						SNew(SSplitter)
 						.Orientation(EOrientation::Orient_Vertical)
 						.Visibility_Lambda([]() { return AreControlledChangelistsEnabled() && AreUncontrolledChangelistsEnabled() ? EVisibility::Visible : EVisibility::Collapsed; })
-				
+						
 						// Top slot: Changelists
 						+SSplitter::Slot()
 						.SizeRule_Lambda([this](){ return ChangelistExpandableArea->IsExpanded() ? SSplitter::ESizeRule::FractionOfParent : SSplitter::ESizeRule::SizeToContent; })
@@ -457,7 +492,7 @@ void SSourceControlChangelistsWidget::Construct(const FArguments& InArgs)
 						[
 							ChangelistExpandableArea.ToSharedRef()
 						]
-
+	
 						// Bottom slot: Uncontrolled Changelists
 						+SSplitter::Slot()
 						.SizeRule_Lambda([this](){ return UncontrolledChangelistExpandableArea->IsExpanded() ? SSplitter::ESizeRule::FractionOfParent : SSplitter::ESizeRule::SizeToContent; })
@@ -466,7 +501,7 @@ void SSourceControlChangelistsWidget::Construct(const FArguments& InArgs)
 							UncontrolledChangelistExpandableArea.ToSharedRef()
 						]
 					]
-					+SOverlay::Slot() // Visibile when controlled changelists are enabled but not the uncontrolled ones.
+					+SOverlay::Slot() // Visible when controlled changelists are enabled but not the uncontrolled ones.
 					[
 						SNew(SBox)
 						.Visibility_Lambda([](){ return AreControlledChangelistsEnabled() && !AreUncontrolledChangelistsEnabled() ? EVisibility::Visible : EVisibility::Collapsed; })
@@ -483,16 +518,16 @@ void SSourceControlChangelistsWidget::Construct(const FArguments& InArgs)
 						]
 					]
 				]
+			]
 
-				// Right slot: Files associated to the selected the changelist/uncontrolled changelist.
-				+SSplitter::Slot()
-				.Resizable(true)
-				.SizeRule(SSplitter::FractionOfParent)
-				.Value_Lambda([this, MaxFileAreaRation]() { return FMath::Min(MaxFileAreaRation, FileAreaSize); })
-				.OnSlotResized_Lambda([this](float Size) { FileAreaSize = Size; })
-				[
-					FileListView.ToSharedRef()
-				]
+			// Right slot: Files associated to the selected the changelist/uncontrolled changelist.
+			+SSplitter::Slot()
+			.Resizable(true)
+			.SizeRule(SSplitter::FractionOfParent)
+			.Value_Lambda([this, MaxFileAreaRation]() { return FMath::Min(MaxFileAreaRation, FileAreaSize); })
+			.OnSlotResized_Lambda([this](float Size) { FileAreaSize = Size; })
+			[
+				FileListViewSwitcher.ToSharedRef()
 			]
 		]
 		+SVerticalBox::Slot() // Status bar (Always visible if uncontrolled changelist are enabled to keep the reconcile status visible at all time)
@@ -647,17 +682,21 @@ void SSourceControlChangelistsWidget::Tick(const FGeometry& AllottedGeometry, co
 			{
 				RequestFileStatusRefresh(*UncontrolledChangelistTreeView->GetSelectedItems()[0]);
 			}
+			else if (UnsavedAssetsTreeView->GetNumItemsSelected() > 0)
+			{
+				RequestFileStatusRefresh(*UnsavedAssetsTreeView->GetSelectedItems()[0]);
+			}
 			bUpdateMonitoredFileStatusList = false;
 		}
 		else // Just monitor the files that are visible in the view.
 		{
 			// NOTE: In the tick where the view is updated/refreshed, the item count is usually zero, next tick it will have the correct value, so wait for it.
-			int32 ItemCount = FileListView->GetNumGeneratedChildren();
+			int32 ItemCount = GetActiveFileListView().GetNumGeneratedChildren();
 			if (ItemCount > 0)
 			{
 				// Minimize the list of file status monitored to only include those in view. The ones outside of the visible area don't need
 				// to be monitored unless the user sorts.
-				int32 ItemStart = static_cast<int32>(FileListView->GetScrollOffset()); // This give the offset in item count.
+				int32 ItemStart = static_cast<int32>(GetActiveFileListView().GetScrollOffset()); // This give the offset in item count.
 				int32 ItemEnd = ItemStart + ItemCount;
 				TSet<FString> MonitoredFiles;
 				for (int32 Index = ItemStart; Index < ItemEnd; ++Index)
@@ -680,7 +719,7 @@ void SSourceControlChangelistsWidget::Tick(const FGeometry& AllottedGeometry, co
 		{
 			for (auto It = TimestampUpdateTask->GetTask().QueriedFileTimestamps.CreateIterator(); It; ++It)
 			{
-				if (TSharedPtr<IFileViewTreeItem>* FileViewItem = OfflineFileItemCache.Find(It->Key))
+				if (TSharedPtr<IFileViewTreeItem>* FileViewItem = OfflineFileItemCache.Find(It->Key); FileViewItem->IsValid())
 				{
 					(*FileViewItem)->SetLastModifiedDateTime(It->Value);
 					bUpdatedItemInView |= (*FileViewItem)->DisplayedUpdateNum == UpdateRequestNum;
@@ -714,7 +753,7 @@ void SSourceControlChangelistsWidget::Tick(const FGeometry& AllottedGeometry, co
 		if (bUpdatedItemInView && IsFileViewSortedByLastModifiedTimestamp())
 		{
 			SortFileView();
-			FileListView->RequestListRefresh();
+			GetActiveFileListView().RequestListRefresh();
 		}
 	}
 	if (!TimestampUpdateTask && OutdatedTimestampFiles.Num()) // Should be sufficient to run only one async task at the time.
@@ -842,17 +881,96 @@ void SSourceControlChangelistsWidget::ClearChangelistsTree()
 		UncontrolledChangelistTreeView->RequestTreeRefresh();
 	}
 
-	if (!FileListNodes.IsEmpty())
-	{
-		FileListNodes.Reset();
-		FileListView->RequestListRefresh();
-	}
-
 	SourceControlItemCache.Reset();
 	UncontrolledChangelistItemCache.Reset();
 	OfflineFileItemCache.Reset();
 }
 
+void SSourceControlChangelistsWidget::OnRefreshUnsavedAssetsWidgets(int64 CurrUpdateNum, const TFunction<void(TSharedPtr<IFileViewTreeItem>&)>& AddItemToFileView)
+{
+	const bool bBeautifyPaths = true;
+	const int64 PrevUpdateNum = CurrUpdateNum - 1;
+	int32 DisplayedItemCount = 0;
+	int32 VisitedItemCount = 0;
+
+	TSharedPtr<IChangelistTreeItem> UnsavedAssetsItem = UnsavedAssetsTreeView->GetNumItemsSelected() > 0 ? UnsavedAssetsTreeView->GetSelectedItems()[0] : TSharedPtr<IChangelistTreeItem>();
+
+	if (!UnsavedAssetsItem)
+	{
+		return;
+	}
+
+	auto OnSourceControlCachedItemVisited = [CurrUpdateNum, &VisitedItemCount](IChangelistTreeItem* Item)
+	{
+		checkf(Item->VisitedUpdateNum < CurrUpdateNum, TEXT("The same revision control item was visited twice. It is likely present in more than one changelist"));
+		Item->VisitedUpdateNum = CurrUpdateNum;
+		++VisitedItemCount;
+	};
+
+	auto OnNewSourceControlItemCreated = [this, &OnSourceControlCachedItemVisited](FString Key, TSharedPtr<IFileViewTreeItem> Item)
+	{
+		OnSourceControlCachedItemVisited(Item.Get());
+		OfflineFileItemCache.Emplace(MoveTemp(Key), MoveTemp(Item));
+	};
+
+	auto AddItemToFileViewWrapper = [&AddItemToFileView, &DisplayedItemCount](TSharedPtr<IFileViewTreeItem>& Item)
+	{
+		++DisplayedItemCount;
+		AddItemToFileView(Item);
+	};
+
+	auto RemoveDiscardedSourceControlItems = [this, VisitedItemCount, DisplayedItemCount, CurrUpdateNum]()
+	{
+		TRACE_CPUPROFILER_EVENT_SCOPE(SSourceControlChangelistsWidget::RemoveDiscardedSourceControlItems);
+
+		check(OfflineFileItemCache.Num() >= VisitedItemCount); // Cannot visit items that are not cached.
+
+		// Remove the items used to be visited but weren't visited during this iteration.
+		int32 DiscardCount = OfflineFileItemCache.Num() - VisitedItemCount;
+		for (auto It = OfflineFileItemCache.CreateIterator(); It && DiscardCount > 0 ; ++It)
+		{
+			if (It->Value->VisitedUpdateNum < CurrUpdateNum)
+			{
+				It.RemoveCurrent();
+				--DiscardCount;
+			}
+		}
+		
+		// Unless the FileTreeNode array was reset because the selected changelist changed, it still contains the union of items to display for this iteration and the item that were displayed last
+		// iteration, so we need to remove item that aren't displayed anymore. (The FileTreeNode is not reset between updates to preserve the sort order as much as possible).
+		if (FileListNodes.Num() > DisplayedItemCount)
+		{
+			// Remove items that were not displayed anymore.
+			FileListNodes.RemoveAll([CurrUpdateNum](const TSharedPtr<IChangelistTreeItem>& Candidate) { return Candidate->DisplayedUpdateNum < CurrUpdateNum; });
+		}
+	};
+
+	for (const FString& UnsavedAssetPath : FUnsavedAssetsTrackerModule::Get().GetUnsavedAssets())
+	{
+		TSharedPtr<IFileViewTreeItem> FileItem;
+		if (TSharedPtr<IFileViewTreeItem>* CachedFileItem = OfflineFileItemCache.Find(UnsavedAssetPath); CachedFileItem->IsValid())
+		{
+			check((*CachedFileItem)->GetTreeItemType() == IChangelistTreeItem::OfflineFile);
+			FileItem = StaticCastSharedPtr<FOfflineFileTreeItem>(*CachedFileItem);
+			OnSourceControlCachedItemVisited(FileItem.Get());
+		}
+		else
+		{
+			FileItem = MakeShared<FOfflineFileTreeItem>(UnsavedAssetPath);
+			OnNewSourceControlItemCreated(UnsavedAssetPath, FileItem);
+		}
+
+		if (FileTextFilter->PassesFilter(*FileItem))
+		{
+			UnsavedAssetsItem->AddChild(FileItem.ToSharedRef());
+			AddItemToFileViewWrapper(FileItem);
+		}
+	}
+
+	RemoveDiscardedSourceControlItems();
+	UnsavedAssetsTreeView->RequestTreeRefresh();
+}
+
 void SSourceControlChangelistsWidget::OnRefreshSourceControlWidgets(int64 CurrUpdateNum, const TFunction<void(TSharedPtr<IFileViewTreeItem>&)>& AddItemToFileView)
 {
 	TRACE_CPUPROFILER_EVENT_SCOPE(SSourceControlChangelistsWidget::OnRefreshSourceControlWidgets);
@@ -932,7 +1050,7 @@ void SSourceControlChangelistsWidget::OnRefreshSourceControlWidgets(int64 CurrUp
 		// Create or reuse the 'Changelists' and 'Shelved Files' node.
 		TSharedPtr<FChangelistTreeItem> ChangelistTreeItem;
 		TSharedPtr<FShelvedChangelistTreeItem> ShelvedChangelistTreeItem;
-		if (TSharedPtr<IChangelistTreeItem>* CachedChangelistItem = SourceControlItemCache.Find(ChangelistState))
+		if (TSharedPtr<IChangelistTreeItem>* CachedChangelistItem = SourceControlItemCache.Find(ChangelistState); CachedChangelistItem->IsValid())
 		{
 			check((*CachedChangelistItem)->GetTreeItemType() == IChangelistTreeItem::Changelist);
 			ChangelistTreeItem = StaticCastSharedPtr<FChangelistTreeItem>(*CachedChangelistItem);
@@ -969,7 +1087,7 @@ void SSourceControlChangelistsWidget::OnRefreshSourceControlWidgets(int64 CurrUp
 		for (const TSharedRef<ISourceControlState>& FileState : ChangelistState->GetFilesStates())
 		{
 			TSharedPtr<IFileViewTreeItem> FileItem;
-			if (TSharedPtr<IChangelistTreeItem>* CachedFileItem = SourceControlItemCache.Find(FileState))
+			if (TSharedPtr<IChangelistTreeItem>* CachedFileItem = SourceControlItemCache.Find(FileState); CachedFileItem->IsValid())
 			{
 				check((*CachedFileItem)->GetTreeItemType() == IChangelistTreeItem::File);
 				FileItem = StaticCastSharedPtr<FFileTreeItem>(*CachedFileItem);
@@ -996,7 +1114,7 @@ void SSourceControlChangelistsWidget::OnRefreshSourceControlWidgets(int64 CurrUp
 		for (const TSharedRef<ISourceControlState>& ShelvedFileState : ChangelistState->GetShelvedFilesStates())
 		{
 			TSharedPtr<IFileViewTreeItem> ShelvedFileItem;
-			if (TSharedPtr<IChangelistTreeItem>* CachedShelvedFileItem = SourceControlItemCache.Find(ShelvedFileState))
+			if (TSharedPtr<IChangelistTreeItem>* CachedShelvedFileItem = SourceControlItemCache.Find(ShelvedFileState); CachedShelvedFileItem->IsValid())
 			{
 				check((*CachedShelvedFileItem)->GetTreeItemType() == IChangelistTreeItem::ShelvedFile);
 				ShelvedFileItem = StaticCastSharedPtr<FShelvedFileTreeItem>(*CachedShelvedFileItem);
@@ -1111,7 +1229,7 @@ void SSourceControlChangelistsWidget::OnRefreshUncontrolledChangelistWidgets(int
 	{
 		// Create or reuse the 'Uncontrolled Changelist' node.
 		TSharedPtr<FUncontrolledChangelistTreeItem> UncontrolledChangelistNode;
-		if (TSharedPtr<IChangelistTreeItem>* CachedUncontrolledChangelistItem = UncontrolledChangelistItemCache.Find(UncontrolledChangelistState))
+		if (TSharedPtr<IChangelistTreeItem>* CachedUncontrolledChangelistItem = UncontrolledChangelistItemCache.Find(UncontrolledChangelistState); CachedUncontrolledChangelistItem->IsValid())
 		{
 			check((*CachedUncontrolledChangelistItem)->GetTreeItemType() == IChangelistTreeItem::UncontrolledChangelist);
 			UncontrolledChangelistNode = StaticCastSharedPtr<FUncontrolledChangelistTreeItem>(*CachedUncontrolledChangelistItem);
@@ -1135,7 +1253,7 @@ void SSourceControlChangelistsWidget::OnRefreshUncontrolledChangelistWidgets(int
 		for (const TSharedRef<ISourceControlState>& FileState : UncontrolledChangelistState->GetFilesStates())
 		{
 			TSharedPtr<IFileViewTreeItem> UncontrolledFileItem;
-			if (TSharedPtr<IChangelistTreeItem>* CachedUncontrolledFileItem = UncontrolledChangelistItemCache.Find(FileState))
+			if (TSharedPtr<IChangelistTreeItem>* CachedUncontrolledFileItem = UncontrolledChangelistItemCache.Find(FileState); CachedUncontrolledFileItem->IsValid())
 			{
 				check((*CachedUncontrolledFileItem)->GetTreeItemType() == IChangelistTreeItem::File);
 				UncontrolledFileItem = StaticCastSharedPtr<FFileTreeItem>(*CachedUncontrolledFileItem);
@@ -1257,6 +1375,11 @@ void SSourceControlChangelistsWidget::OnRefreshUI(ERefreshFlags RefreshFlag)
 		OnRefreshUncontrolledChangelistWidgets(CurrUpdateNum, AddItemToFileView);
 		bFileViewUpdated |= (UncontrolledChangelistTreeView->GetNumItemsSelected() > 0);
 	}
+	if (EnumHasAnyFlags(RefreshFlag, ERefreshFlags::UnsavedAssets))
+	{
+		OnRefreshUnsavedAssetsWidgets(CurrUpdateNum, AddItemToFileView);
+		bFileViewUpdated |= (UnsavedAssetsTreeView->GetNumItemsSelected() > 0);
+	}
 
 	if (!bFileViewUpdated)
 	{
@@ -1278,8 +1401,8 @@ void SSourceControlChangelistsWidget::OnRefreshUI(ERefreshFlags RefreshFlag)
 	{
 		SortFileView();
 	}
-
-	FileListView->RequestListRefresh();
+	
+	GetActiveFileListView().RequestListRefresh();
 }
 
 void SSourceControlChangelistsWidget::OnSourceControlProviderChanged(ISourceControlProvider& OldProvider, ISourceControlProvider& NewProvider)
@@ -1295,7 +1418,7 @@ void SSourceControlChangelistsWidget::OnSourceControlProviderChanged(ISourceCont
 		if (ChangelistTreeView->GetNumItemsSelected() > 0)
 		{
 			FileListNodes.Reset();
-			FileListView->RequestListRefresh();
+			GetActiveFileListView().RequestListRefresh();
 		}
 		ChangelistTreeNodes.Reset();
 		SourceControlItemCache.Reset();
@@ -1461,7 +1584,7 @@ void SSourceControlChangelistsWidget::SetSelectedFiles(const TArray<FString>& Fi
 	}
 
 	// Clear the current selection.
-	FileListView->ClearSelection();
+	GetActiveFileListView().ClearSelection();
 	TSharedPtr<IChangelistTreeItem> LastFoundItem;
 
 	// If we found a Changelist, select files.
@@ -1478,7 +1601,7 @@ void SSourceControlChangelistsWidget::SetSelectedFiles(const TArray<FString>& Fi
 				const FString& ChildFilename = StaticCastSharedPtr<FFileTreeItem>(ChildItem)->FileState->GetFilename().ToLower();
 				if (FilenamesLowerCase.Contains(ChildFilename))
 				{
-					FileListView->SetItemSelection(ChildItem, /*bSelected*/true);
+					GetActiveFileListView().SetItemSelection(ChildItem, /*bSelected*/true);
 					LastFoundItem = ChildItem;
 				}
 			}
@@ -1487,7 +1610,7 @@ void SSourceControlChangelistsWidget::SetSelectedFiles(const TArray<FString>& Fi
 				const FString& ChildFilename = StaticCastSharedPtr<FOfflineFileTreeItem>(ChildItem)->GetFilename().ToLower();
 				if (FilenamesLowerCase.Contains(ChildFilename))
 				{
-					FileListView->SetItemSelection(ChildItem, /*bSelected*/true);
+					GetActiveFileListView().SetItemSelection(ChildItem, /*bSelected*/true);
 					LastFoundItem = ChildItem;
 				}
 			}
@@ -1515,35 +1638,38 @@ void SSourceControlChangelistsWidget::SetSelectedFiles(const TArray<FString>& Fi
 					UncontrolledChangelistTreeView->SetSelection(FoundChangelistTreeItem);
 				}
 			}
-			FileListView->RequestScrollIntoView(LastFoundItem);
+			GetActiveFileListView().RequestScrollIntoView(LastFoundItem);
 		}
 	}
 }
 
 bool SSourceControlChangelistsWidget::HasFilesSelected() const
 {
-	if (FileListView->GetNumItemsSelected() == 0)
+	if (GetActiveFileListView().GetNumItemsSelected() == 0)
 	{
 		return false;
 	}
-	if (ChangelistTreeView->GetNumItemsSelected() == 0)
+	if (ChangelistTreeView->GetNumItemsSelected() > 0)
 	{
-		return false;
-	}
-	switch (ChangelistTreeView->GetSelectedItems()[0]->GetTreeItemType())
-	{
-	case IChangelistTreeItem::Changelist:
-		return true; // If this type of changeslist is selected and the file view as item selected, it means files are selected.
-	case IChangelistTreeItem::UncontrolledChangelist:
-		for (const TSharedPtr<IChangelistTreeItem>& Item : FileListView->GetSelectedItems())
+		switch (ChangelistTreeView->GetSelectedItems()[0]->GetTreeItemType())
 		{
-			if (Item->GetTreeItemType() == IChangelistTreeItem::File)
-			{
-				return true; // Early out.
-			}
+		case IChangelistTreeItem::Changelist:
+			return true; // If this type of changeslist is selected and the file view as item selected, it means files are selected.
+			case IChangelistTreeItem::UncontrolledChangelist:
+				for (const TSharedPtr<IChangelistTreeItem>& Item : FileListView->GetSelectedItems())
+				{
+					if (Item->GetTreeItemType() == IChangelistTreeItem::File)
+					{
+						return true; // Early out.
+					}
+				}
+		default:
+			break;
 		}
-	default:
-		return false;
+	}
+	if (UnsavedAssetsFileListView->GetSelectedItems().Num() > 0)
+	{
+		return true;
 	}
 
 	return false;
@@ -1565,7 +1691,7 @@ TArray<FString> SSourceControlChangelistsWidget::GetSelectedFiles()
 
 void SSourceControlChangelistsWidget::GetSelectedFiles(TArray<FString>& OutControlledFiles, TArray<FString>& OutUncontrolledFiles)
 {
-	for (const TSharedPtr<IChangelistTreeItem>& Item : FileListView->GetSelectedItems())
+	for (const TSharedPtr<IChangelistTreeItem>& Item : GetActiveFileListView().GetSelectedItems())
 	{
 		if (Item->GetTreeItemType() == IChangelistTreeItem::File)
 		{
@@ -1592,6 +1718,11 @@ void SSourceControlChangelistsWidget::GetSelectedFiles(TArray<FString>& OutContr
 					const FString& Filename = StaticCastSharedPtr<FOfflineFileTreeItem>(Item)->GetFilename();
 					OutUncontrolledFiles.Add(Filename);
 				}
+				else if (Parent->GetTreeItemType() == IChangelistTreeItem::UnsavedAssets)
+				{
+					const FString& Filename = StaticCastSharedPtr<FOfflineFileTreeItem>(Item)->GetFilename();
+					OutControlledFiles.Add(Filename);
+				}
 			}
 		}
 	}
@@ -2699,10 +2830,12 @@ TSharedPtr<SWidget> SSourceControlChangelistsWidget::OnOpenContextMenu()
 
 	TArray<TSharedPtr<IChangelistTreeItem>> SelectedChangelistNodes = ChangelistTreeView->GetSelectedItems();
 	TArray<TSharedPtr<IChangelistTreeItem>> SelectedUncontrolledChangelistNodes = UncontrolledChangelistTreeView->GetSelectedItems();
+	TArray<TSharedPtr<IChangelistTreeItem>> SelectedUnsavedAssetNodes = UnsavedAssetsTreeView->GetSelectedItems();
 
 	bool bHasSelectedChangelist = SelectedChangelistNodes.Num() > 0 &&  SelectedChangelistNodes[0]->GetTreeItemType() == IChangelistTreeItem::Changelist;
 	bool bHasSelectedShelvedChangelistNode = SelectedChangelistNodes.Num() > 0 &&  SelectedChangelistNodes[0]->GetTreeItemType() == IChangelistTreeItem::ShelvedChangelist;
 	bool bHasSelectedUncontrolledChangelist = SelectedUncontrolledChangelistNodes.Num() > 0 &&  SelectedUncontrolledChangelistNodes[0]->GetTreeItemType() == IChangelistTreeItem::UncontrolledChangelist;
+	bool bHasSelectedUnsavedAsset = SelectedUnsavedAssetNodes.Num() > 0;
 	bool bHasSelectedFiles = HasFilesSelected();
 	bool bHasSelectedShelvedFiles = HasShelvedFilesSelected();
 	bool bHasEmptySelection = (!bHasSelectedChangelist && !bHasSelectedFiles && !bHasSelectedShelvedFiles);
@@ -2722,6 +2855,63 @@ TSharedPtr<SWidget> SSourceControlChangelistsWidget::OnOpenContextMenu()
 
 	FToolMenuSection& Section = *Menu->FindSection("Source Control");
 	
+	if (bHasSelectedUnsavedAsset)
+	{
+		Section.AddMenuEntry(
+			"Save", LOCTEXT("SourceControl_SaveUnsavedAsset", "Save"),
+			FText::FormatOrdered(LOCTEXT("SourceControl_SaveUnsavedAsset_Tooltip", "Save unsaved {0}|plural(one=asset,other=assets)"), ContextObject->SelectedFiles.Num()),
+			FSlateIcon(FAppStyle::GetAppStyleSetName(), "Icons.Save"),
+			FUIAction(FExecuteAction::CreateLambda([this]
+			{
+				TArray<FString> SelectedFiles;
+				GetSelectedFiles(SelectedFiles, SelectedFiles);
+				TArray<UPackage*> Packages;
+				for (const FString& Filename : SelectedFiles)
+				{
+					FString PackageName = UPackageTools::FilenameToPackageName(Filename);
+					Packages.Add(UPackageTools::LoadPackage(PackageName));
+				}
+				const bool bOnlyDirty = false;
+				UEditorLoadingAndSavingUtils::SavePackages(Packages, bOnlyDirty);
+			})));
+
+		Section.AddMenuEntry(
+			"SaveToChangelist", LOCTEXT("SourceControl_SaveUnsavedAssetToChangelist", "Save To Changelist"),
+			FText::FormatOrdered(LOCTEXT("SourceControl_SaveUnsavedAssetToChangelist_Tooltip", "Save {0}|plural(one=this asset,other=these assets) and add {0}|plural(one=it,other=them) to a changelist"), ContextObject->SelectedFiles.Num()),
+			FSlateIcon(FAppStyle::GetAppStyleSetName(), "Icons.ArrowRight"),
+			FUIAction(FExecuteAction::CreateLambda([this]
+			{
+				FPackageSourceControlHelper PackageHelper;
+				
+				TArray<FString> SelectedFiles;
+				GetSelectedFiles(SelectedFiles, SelectedFiles);
+				TArray<FString> PackageNames;
+				for (const FString& Filename : SelectedFiles)
+				{
+					PackageNames.Add(UPackageTools::FilenameToPackageName(Filename));
+				}
+				
+				if (!PackageHelper.Checkout(PackageNames))
+				{
+					return;
+				}
+				
+				// Handles almost everything else
+				OnMoveFiles();
+
+				// Actaully save the assets 
+				TArray<UPackage*> Packages;
+				for (const FString& PackageName : PackageNames)
+				{
+					Packages.Add(UPackageTools::LoadPackage(PackageName));
+				}
+				const bool bOnlyDirty = false;
+				UEditorLoadingAndSavingUtils::SavePackages(Packages, bOnlyDirty);
+			})));
+
+		Section.AddSeparator("ActionSeparator");
+	}
+	
 	// This should appear only on change lists
 	if (bHasSelectedChangelist)
 	{
@@ -2855,7 +3045,7 @@ TSharedPtr<SWidget> SSourceControlChangelistsWidget::OnOpenContextMenu()
 	}
 
 	// Files-only operations
-	if(bHasSelectedFiles)
+	if(bHasSelectedFiles && !bHasSelectedUnsavedAsset)
 	{
 		Section.AddSeparator("FilesSeparator");
 
@@ -2882,7 +3072,7 @@ TSharedPtr<SWidget> SSourceControlChangelistsWidget::OnOpenContextMenu()
 				FCanExecuteAction::CreateSP(this, &SSourceControlChangelistsWidget::CanDiffAgainstDepot)));
 	}
 
-	if (FUncontrolledChangelistsModule::Get().IsEnabled())
+	if (FUncontrolledChangelistsModule::Get().IsEnabled() && !bHasSelectedUnsavedAsset)
 	{
 		Section.AddSeparator("ReconcileSeparator");
 
@@ -3011,7 +3201,64 @@ TSharedRef<SListView<FChangelistTreeItemPtr>> SSourceControlChangelistsWidget::C
 			.SortPriority(this, &SSourceControlChangelistsWidget::GetColumnSortPriority, SourceControlFileViewColumn::CheckedOutByUser::Id())
 			.SortMode(this, &SSourceControlChangelistsWidget::GetColumnSortMode, SourceControlFileViewColumn::CheckedOutByUser::Id())
 			.OnSort(this, &SSourceControlChangelistsWidget::OnColumnSortModeChanged));
+	
+	return FileView;
+}
+
+TSharedRef<SListView<FChangelistTreeItemPtr>> SSourceControlChangelistsWidget::CreateUnsavedAssetsFilesView()
+{
+	TSharedRef<SListView<FChangelistTreeItemPtr>> FileView = SNew(SListView<FChangelistTreeItemPtr>)
+		.ItemHeight(24.0f)
+		.ListItemsSource(&FileListNodes)
+		.OnGenerateRow(this, &SSourceControlChangelistsWidget::OnGenerateRow)
+		.SelectionMode(ESelectionMode::Multi)
+		.OnContextMenuOpening(this, &SSourceControlChangelistsWidget::OnOpenContextMenu)
+		.HeaderRow
+		(
+			SNew(SHeaderRow)
+			.CanSelectGeneratedColumn(true)
+			.HiddenColumnsList(FileViewHiddenColumnsList)
+			.OnHiddenColumnsListChanged(this, &SSourceControlChangelistsWidget::OnFileViewHiddenColumnsListChanged)
+
+			+SHeaderRow::Column(SourceControlFileViewColumn::Name::Id())
+			.DefaultLabel(SourceControlFileViewColumn::Name::GetDisplayText())
+			.DefaultTooltip(SourceControlFileViewColumn::Name::GetToolTipText())
+			.ShouldGenerateWidget(true) // Ensure the column cannot be hidden (grayed out in the show/hide drop down menu)
+			.FillWidth(1.5f)
+			.SortPriority(this, &SSourceControlChangelistsWidget::GetColumnSortPriority, SourceControlFileViewColumn::Name::Id())
+			.SortMode(this, &SSourceControlChangelistsWidget::GetColumnSortMode, SourceControlFileViewColumn::Name::Id())
+			.OnSort(this, &SSourceControlChangelistsWidget::OnColumnSortModeChanged)
 
+			+SHeaderRow::Column(SourceControlFileViewColumn::Path::Id())
+			.DefaultLabel(SourceControlFileViewColumn::Path::GetDisplayText())
+			.DefaultTooltip(SourceControlFileViewColumn::Path::GetToolTipText())
+			.ShouldGenerateWidget(true) // Ensure the column cannot be hidden (grayed out in the show/hide drop down menu)
+			.FillWidth(3.5f)
+			.SortPriority(this, &SSourceControlChangelistsWidget::GetColumnSortPriority, SourceControlFileViewColumn::Path::Id())
+			.SortMode(this, &SSourceControlChangelistsWidget::GetColumnSortMode, SourceControlFileViewColumn::Path::Id())
+			.OnSort(this, &SSourceControlChangelistsWidget::OnColumnSortModeChanged)
+
+			+SHeaderRow::Column(SourceControlFileViewColumn::Type::Id())
+			.DefaultLabel(SourceControlFileViewColumn::Type::GetDisplayText())
+			.DefaultTooltip(SourceControlFileViewColumn::Type::GetToolTipText())
+			.ShouldGenerateWidget(true) // Ensure the column cannot be hidden (grayed out in the show/hide drop down menu)
+			.FillWidth(3.5f)
+			.SortPriority(this, &SSourceControlChangelistsWidget::GetColumnSortPriority, SourceControlFileViewColumn::Type::Id())
+			.SortMode(this, &SSourceControlChangelistsWidget::GetColumnSortMode, SourceControlFileViewColumn::Type::Id())
+			.OnSort(this, &SSourceControlChangelistsWidget::OnColumnSortModeChanged)
+
+			+SHeaderRow::Column(SourceControlFileViewColumn::Discard::Id())
+			.DefaultLabel(SourceControlFileViewColumn::Discard::GetDisplayText())
+			.DefaultTooltip(SourceControlFileViewColumn::Discard::GetToolTipText())
+			.ShouldGenerateWidget(true) // Ensure the column cannot be hidden (grayed out in the show/hide drop down menu)
+			.FixedWidth(18)
+			.HeaderContentPadding(FMargin(0))
+			.HAlignHeader(HAlign_Right)
+			[
+				SNew(SBox)
+					.WidthOverride(18)
+					.HeightOverride(16)
+			]);
 	return FileView;
 }
 
@@ -3083,6 +3330,9 @@ TSharedRef<ITableRow> SSourceControlChangelistsWidget::OnGenerateRow(TSharedPtr<
 			.TreeItemToVisualize(InTreeItem)
 			.HighlightText_Lambda([this]() { return FileSearchBox->GetText(); });
 
+	case IChangelistTreeItem::UnsavedAssets:
+		return SNew(SUnsavedAssetsTableRow, OwnerTable);
+		
 	default:
 		checkNoEntry();
 	};
@@ -3136,6 +3386,10 @@ void SSourceControlChangelistsWidget::OnGetChangelistChildren(FChangelistTreeIte
 	{
 		// Uncontrolled changelist nodes do not have children at the moment.
 	}
+	else if (InParent->GetTreeItemType() == IChangelistTreeItem::UnsavedAssets)
+	{
+		// The unsaved assets node does not have children at the moment.
+	}
 }
 
 void SSourceControlChangelistsWidget::OnItemDoubleClicked(TSharedPtr<IChangelistTreeItem> Item)
@@ -3181,14 +3435,25 @@ void SSourceControlChangelistsWidget::OnChangelistSelectionChanged(TSharedPtr<IC
 			case IChangelistTreeItem::Changelist:
 			case IChangelistTreeItem::ShelvedChangelist:
 				UncontrolledChangelistTreeView->ClearSelection(); // Don't have a changelists selected at the same time than an uncontrolled one, they share the same file view.
+				UnsavedAssetsTreeView->ClearSelection();
+				FileListViewSwitcher->SetActiveWidget(FileListView.ToSharedRef());
 				OnRefreshUI(ERefreshFlags::SourceControlChangelists);
 				break;
 
 			case IChangelistTreeItem::UncontrolledChangelist:
 				ChangelistTreeView->ClearSelection();
+				UnsavedAssetsTreeView->ClearSelection();
+				FileListViewSwitcher->SetActiveWidget(FileListView.ToSharedRef());
 				OnRefreshUI(ERefreshFlags::UncontrolledChangelists);
 				break;
 
+			case IChangelistTreeItem::UnsavedAssets:
+				ChangelistTreeView->ClearSelection();
+				UncontrolledChangelistTreeView->ClearSelection();
+				FileListViewSwitcher->SetActiveWidget(UnsavedAssetsFileListView.ToSharedRef());
+				OnRefreshUI(ERefreshFlags::UnsavedAssets);
+				break;
+			
 			default:
 				break;
 		}
@@ -3443,9 +3708,19 @@ void SSourceControlChangelistsWidget::PopulateItemSearchStrings(const IChangelis
 	}
 }
 
+void SSourceControlChangelistsWidget::OnUnsavedAssetChanged(const FString& /*Filepath*/)
+{
+	OnRefreshUI(ERefreshFlags::UnsavedAssets);
+}
+
 void SSourceControlChangelistsWidget::OnPackageSaved(const FString& PackageFilename, UPackage* Package, FObjectPostSaveContext ObjectSaveContext)
 {
 	OutdatedTimestampFiles.Emplace(PackageFilename); // The background task will convert relative to absolute path.
 }
 
+SListView<FChangelistTreeItemPtr>& SSourceControlChangelistsWidget::GetActiveFileListView() const
+{
+	return *StaticCastSharedPtr<SListView<FChangelistTreeItemPtr>>(FileListViewSwitcher->GetActiveWidget());
+}
+
 #undef LOCTEXT_NAMESPACE
diff --git a/Engine/Source/Editor/SourceControlWindows/Private/SSourceControlChangelists.h b/Engine/Source/Editor/SourceControlWindows/Private/SSourceControlChangelists.h
--- a/Engine/Source/Editor/SourceControlWindows/Private/SSourceControlChangelists.h
+++ b/Engine/Source/Editor/SourceControlWindows/Private/SSourceControlChangelists.h
@@ -17,7 +17,9 @@
 #include "Widgets/Views/STreeView.h"
 
 class FChangelistGroupTreeItem;
+class SExpandableArea;
 class SExpandableChangelistArea;
+class SWidgetSwitcher;
 class SSearchBox;
 class USourceControlSettings;
 
@@ -41,8 +43,9 @@ private:
 	enum class ERefreshFlags
 	{
 		SourceControlChangelists = 1 << 0,
+		UnsavedAssets            = 1 << 1,
 		UncontrolledChangelists  = 1 << 2,
-		All = SourceControlChangelists | UncontrolledChangelists,
+		All = SourceControlChangelists | UnsavedAssets | UncontrolledChangelists,
 	};
 	FRIEND_ENUM_CLASS_FLAGS(ERefreshFlags)
 
@@ -65,6 +68,7 @@ private:
 	TSharedRef<SWidget> MakeToolBar();
 	TSharedRef<STreeView<FChangelistTreeItemPtr>> CreateChangelistTreeView(TArray<TSharedPtr<IChangelistTreeItem>>& ItemSources);
 	TSharedRef<SListView<FChangelistTreeItemPtr>> CreateChangelistFilesView();
+	TSharedRef<SListView<FChangelistTreeItemPtr>> CreateUnsavedAssetsFilesView();
 
 	TSharedRef<ITableRow> OnGenerateRow(FChangelistTreeItemPtr InTreeItem, const TSharedRef<STableViewBase>& OwnerTable);
 	void OnGetChangelistChildren(FChangelistTreeItemPtr InParent, TArray<FChangelistTreeItemPtr>& OutChildren);
@@ -81,6 +85,7 @@ private:
 	void OnUncontrolledChangelistSearchTextChanged(const FText& InFilterText);
 	void OnFileSearchTextChanged(const FText& InFilterText);
 	void PopulateItemSearchStrings(const IChangelistTreeItem& Item, TArray<FString>& OutStrings);
+	void OnUnsavedAssetChanged(const FString& Filepath);
 
 	FReply OnFilesDragged(const FGeometry& InGeometry, const FPointerEvent& InMouseEvent);
 
@@ -88,6 +93,7 @@ private:
 	void RequestFileStatusRefresh(const IChangelistTreeItem& Changelist);
 	void RequestFileStatusRefresh(TSet<FString>&& Pathnames);
 	void OnRefreshUI(ERefreshFlags RefreshFlags);
+	void OnRefreshUnsavedAssetsWidgets(int64 CurrUpdateNum, const TFunction<void(TSharedPtr<IFileViewTreeItem>&)>& AddItemToFileView);
 	void OnRefreshSourceControlWidgets(int64 CurreUpdateNum, const TFunction<void(TSharedPtr<IFileViewTreeItem>&)>& AddItemToFileView);
 	void OnRefreshUncontrolledChangelistWidgets(int64 CurreUpdateNum, const TFunction<void(TSharedPtr<IFileViewTreeItem>&)>& AddItemToFileView);
 	void ClearChangelistsTree();
@@ -196,6 +202,10 @@ private:
 	void OnPackageSaved(const FString& PackageFilename, UPackage* Package, FObjectPostSaveContext ObjectSaveContext);
 
 private:
+	TSharedPtr<SExpandableChangelistArea> UnsavedAssetsExpandableArea;
+	TSharedPtr<STreeView<FChangelistTreeItemPtr>> UnsavedAssetsTreeView;
+	TArray<TSharedPtr<IChangelistTreeItem>> UnsavedAssetsTreeNode;
+	
 	TSharedPtr<SExpandableChangelistArea> ChangelistExpandableArea;
 	TSharedPtr<STreeView<FChangelistTreeItemPtr>> ChangelistTreeView;
 	TArray<TSharedPtr<IChangelistTreeItem>> ChangelistTreeNodes;
@@ -203,11 +213,16 @@ private:
 	TSharedPtr<SExpandableChangelistArea> UncontrolledChangelistExpandableArea;
 	TSharedPtr<STreeView<FChangelistTreeItemPtr>> UncontrolledChangelistTreeView;
 	TArray<TSharedPtr<IChangelistTreeItem>> UncontrolledChangelistTreeNodes;
-
+	
 	TSharedPtr<SListView<FChangelistTreeItemPtr>> FileListView;
 	TArray<TSharedPtr<IChangelistTreeItem>> FileListNodes;
 	TArray<FName> FileViewHiddenColumnsList;
 
+	TSharedPtr<SListView<FChangelistTreeItemPtr>> UnsavedAssetsFileListView;
+	TSharedPtr<SWidgetSwitcher> FileListViewSwitcher;
+	
+	SListView<FChangelistTreeItemPtr>& GetActiveFileListView() const;
+
 	TMap<TSharedPtr<void>, TSharedPtr<IChangelistTreeItem>> SourceControlItemCache;
 	TMap<TSharedPtr<void>, TSharedPtr<IChangelistTreeItem>> UncontrolledChangelistItemCache;
 	TMap<FString, TSharedPtr<IFileViewTreeItem>> OfflineFileItemCache;
diff --git a/Engine/Source/Editor/SourceControlWindows/Private/SSourceControlCommon.cpp b/Engine/Source/Editor/SourceControlWindows/Private/SSourceControlCommon.cpp
--- a/Engine/Source/Editor/SourceControlWindows/Private/SSourceControlCommon.cpp
+++ b/Engine/Source/Editor/SourceControlWindows/Private/SSourceControlCommon.cpp
@@ -205,6 +205,11 @@ void IFileViewTreeItem::SetLastModifiedDateTime(const FDateTime& Timestamp)
 
 //////////////////////////////////////////////////////////////////////////
 
+FString FUnsavedAssetsTreeItem::GetDisplayString() const
+{
+	return "";
+}
+
 FFileTreeItem::FFileTreeItem(FSourceControlStateRef InFileState, bool bBeautifyPaths, bool bIsShelvedFile)
 	: IFileViewTreeItem(bIsShelvedFile ? IChangelistTreeItem::ShelvedFile : IChangelistTreeItem::File)
 	, FileState(InFileState)
diff --git a/Engine/Source/Editor/SourceControlWindows/Private/SSourceControlCommon.h b/Engine/Source/Editor/SourceControlWindows/Private/SSourceControlCommon.h
--- a/Engine/Source/Editor/SourceControlWindows/Private/SSourceControlCommon.h
+++ b/Engine/Source/Editor/SourceControlWindows/Private/SSourceControlCommon.h
@@ -26,6 +26,9 @@ struct FAssetData;
  * > UncontrolledChangelist
  *     File
  *     Offline File
+ *
+ * > UnsavedAssets
+ *     Offline File
  */
 struct IChangelistTreeItem : TSharedFromThis<IChangelistTreeItem>
 {
@@ -37,6 +40,7 @@ struct IChangelistTreeItem : TSharedFromThis<IChangelistTreeItem>
 		ShelvedChangelist,       // Node displaying shelved files as children.
 		ShelvedFile,             // Node displaying a shelved file information.
 		OfflineFile,             // Node displaying an offline file information.
+		UnsavedAssets,           // Node displaying unsaved asset category
 	};
 
 	virtual ~IChangelistTreeItem() = default;
@@ -201,6 +205,16 @@ struct FUncontrolledChangelistTreeItem : public IChangelistTreeItem
 	FUncontrolledChangelistStateRef UncontrolledChangelistState;
 };
 
+/** Displays a changelist icon/number/description. */
+struct FUnsavedAssetsTreeItem : public IChangelistTreeItem
+{
+	FUnsavedAssetsTreeItem()
+		: IChangelistTreeItem(IChangelistTreeItem::UnsavedAssets)
+	{
+	}
+
+	FString GetDisplayString() const;
+};
 
 /** Displays a set of files under a changelist or uncontrolled changelist. */
 struct FFileTreeItem : public IFileViewTreeItem
diff --git a/Engine/Source/Editor/SourceControlWindows/SourceControlWindows.Build.cs b/Engine/Source/Editor/SourceControlWindows/SourceControlWindows.Build.cs
--- a/Engine/Source/Editor/SourceControlWindows/SourceControlWindows.Build.cs
+++ b/Engine/Source/Editor/SourceControlWindows/SourceControlWindows.Build.cs
@@ -16,6 +16,7 @@ public class SourceControlWindows : ModuleRules
 				"SlateCore",
 				"SourceControl", 
 				"UncontrolledChangelists",
+				"UnsavedAssetsTracker",
 				"AssetTools",
 				"ToolWidgets",
 				"EditorFramework",
diff --git a/Engine/Source/Runtime/SlateCore/Private/Styling/StarshipCoreStyle.cpp b/Engine/Source/Runtime/SlateCore/Private/Styling/StarshipCoreStyle.cpp
--- a/Engine/Source/Runtime/SlateCore/Private/Styling/StarshipCoreStyle.cpp
+++ b/Engine/Source/Runtime/SlateCore/Private/Styling/StarshipCoreStyle.cpp
@@ -338,7 +338,9 @@ TSharedRef<ISlateStyle> FStarshipCoreStyle::Create()
 		Style->Set("Icons.Plus", new IMAGE_BRUSH_SVG("Starship/Common/plus", Icon16x16));
 		Style->Set("Icons.Minus", new IMAGE_BRUSH_SVG("Starship/Common/minus", Icon16x16));
 		Style->Set("Icons.PlusCircle", new IMAGE_BRUSH_SVG("Starship/Common/plus-circle", Icon16x16));
+		Style->Set("Icons.MinusCircle", new IMAGE_BRUSH_SVG("Starship/Common/minus-circle", Icon16x16));
 		Style->Set("Icons.X", new IMAGE_BRUSH_SVG("Starship/Common/close", Icon16x16));
+		Style->Set("Icons.XCircle", new IMAGE_BRUSH_SVG("Starship/Common/close-circle", Icon16x16));
 		Style->Set("Icons.Delete", new IMAGE_BRUSH_SVG("Starship/Common/delete-outline", Icon16x16));
 		Style->Set("Icons.Save", new IMAGE_BRUSH_SVG("Starship/Common/save", Icon16x16));
 		Style->Set("Icons.SaveModified", new IMAGE_BRUSH_SVG("Starship/Common/save-modified", Icon16x16));
