Fix crash in View Changes UI

diff --git a/Engine/Source/Editor/SourceControlWindows/Private/SSourceControlChangelists.cpp b/Engine/Source/Editor/SourceControlWindows/Private/SSourceControlChangelists.cpp
--- a/Engine/Source/Editor/SourceControlWindows/Private/SSourceControlChangelists.cpp
+++ b/Engine/Source/Editor/SourceControlWindows/Private/SSourceControlChangelists.cpp
@@ -719,7 +719,7 @@ void SSourceControlChangelistsWidget::Tick(const FGeometry& AllottedGeometry, co
 		{
 			for (auto It = TimestampUpdateTask->GetTask().QueriedFileTimestamps.CreateIterator(); It; ++It)
 			{
-				if (TSharedPtr<IFileViewTreeItem>* FileViewItem = OfflineFileItemCache.Find(It->Key); FileViewItem->IsValid())
+				if (TSharedPtr<IFileViewTreeItem>* FileViewItem = OfflineFileItemCache.Find(It->Key); FileViewItem && FileViewItem->IsValid())
 				{
 					(*FileViewItem)->SetLastModifiedDateTime(It->Value);
 					bUpdatedItemInView |= (*FileViewItem)->DisplayedUpdateNum == UpdateRequestNum;
@@ -948,7 +948,7 @@ void SSourceControlChangelistsWidget::OnRefreshUnsavedAssetsWidgets(int64 CurrUp
 	for (const FString& UnsavedAssetPath : FUnsavedAssetsTrackerModule::Get().GetUnsavedAssets())
 	{
 		TSharedPtr<IFileViewTreeItem> FileItem;
-		if (TSharedPtr<IFileViewTreeItem>* CachedFileItem = OfflineFileItemCache.Find(UnsavedAssetPath); CachedFileItem->IsValid())
+		if (TSharedPtr<IFileViewTreeItem>* CachedFileItem = OfflineFileItemCache.Find(UnsavedAssetPath); CachedFileItem && CachedFileItem->IsValid())
 		{
 			check((*CachedFileItem)->GetTreeItemType() == IChangelistTreeItem::OfflineFile);
 			FileItem = StaticCastSharedPtr<FOfflineFileTreeItem>(*CachedFileItem);
@@ -1050,7 +1050,7 @@ void SSourceControlChangelistsWidget::OnRefreshSourceControlWidgets(int64 CurrUp
 		// Create or reuse the 'Changelists' and 'Shelved Files' node.
 		TSharedPtr<FChangelistTreeItem> ChangelistTreeItem;
 		TSharedPtr<FShelvedChangelistTreeItem> ShelvedChangelistTreeItem;
-		if (TSharedPtr<IChangelistTreeItem>* CachedChangelistItem = SourceControlItemCache.Find(ChangelistState); CachedChangelistItem->IsValid())
+		if (TSharedPtr<IChangelistTreeItem>* CachedChangelistItem = SourceControlItemCache.Find(ChangelistState); CachedChangelistItem && CachedChangelistItem->IsValid())
 		{
 			check((*CachedChangelistItem)->GetTreeItemType() == IChangelistTreeItem::Changelist);
 			ChangelistTreeItem = StaticCastSharedPtr<FChangelistTreeItem>(*CachedChangelistItem);
@@ -1087,7 +1087,7 @@ void SSourceControlChangelistsWidget::OnRefreshSourceControlWidgets(int64 CurrUp
 		for (const TSharedRef<ISourceControlState>& FileState : ChangelistState->GetFilesStates())
 		{
 			TSharedPtr<IFileViewTreeItem> FileItem;
-			if (TSharedPtr<IChangelistTreeItem>* CachedFileItem = SourceControlItemCache.Find(FileState); CachedFileItem->IsValid())
+			if (TSharedPtr<IChangelistTreeItem>* CachedFileItem = SourceControlItemCache.Find(FileState); CachedFileItem && CachedFileItem->IsValid())
 			{
 				check((*CachedFileItem)->GetTreeItemType() == IChangelistTreeItem::File);
 				FileItem = StaticCastSharedPtr<FFileTreeItem>(*CachedFileItem);
@@ -1114,7 +1114,7 @@ void SSourceControlChangelistsWidget::OnRefreshSourceControlWidgets(int64 CurrUp
 		for (const TSharedRef<ISourceControlState>& ShelvedFileState : ChangelistState->GetShelvedFilesStates())
 		{
 			TSharedPtr<IFileViewTreeItem> ShelvedFileItem;
-			if (TSharedPtr<IChangelistTreeItem>* CachedShelvedFileItem = SourceControlItemCache.Find(ShelvedFileState); CachedShelvedFileItem->IsValid())
+			if (TSharedPtr<IChangelistTreeItem>* CachedShelvedFileItem = SourceControlItemCache.Find(ShelvedFileState); CachedShelvedFileItem && CachedShelvedFileItem->IsValid())
 			{
 				check((*CachedShelvedFileItem)->GetTreeItemType() == IChangelistTreeItem::ShelvedFile);
 				ShelvedFileItem = StaticCastSharedPtr<FShelvedFileTreeItem>(*CachedShelvedFileItem);
@@ -1229,7 +1229,7 @@ void SSourceControlChangelistsWidget::OnRefreshUncontrolledChangelistWidgets(int
 	{
 		// Create or reuse the 'Uncontrolled Changelist' node.
 		TSharedPtr<FUncontrolledChangelistTreeItem> UncontrolledChangelistNode;
-		if (TSharedPtr<IChangelistTreeItem>* CachedUncontrolledChangelistItem = UncontrolledChangelistItemCache.Find(UncontrolledChangelistState); CachedUncontrolledChangelistItem->IsValid())
+		if (TSharedPtr<IChangelistTreeItem>* CachedUncontrolledChangelistItem = UncontrolledChangelistItemCache.Find(UncontrolledChangelistState); CachedUncontrolledChangelistItem && CachedUncontrolledChangelistItem->IsValid())
 		{
 			check((*CachedUncontrolledChangelistItem)->GetTreeItemType() == IChangelistTreeItem::UncontrolledChangelist);
 			UncontrolledChangelistNode = StaticCastSharedPtr<FUncontrolledChangelistTreeItem>(*CachedUncontrolledChangelistItem);
@@ -1253,7 +1253,7 @@ void SSourceControlChangelistsWidget::OnRefreshUncontrolledChangelistWidgets(int
 		for (const TSharedRef<ISourceControlState>& FileState : UncontrolledChangelistState->GetFilesStates())
 		{
 			TSharedPtr<IFileViewTreeItem> UncontrolledFileItem;
-			if (TSharedPtr<IChangelistTreeItem>* CachedUncontrolledFileItem = UncontrolledChangelistItemCache.Find(FileState); CachedUncontrolledFileItem->IsValid())
+			if (TSharedPtr<IChangelistTreeItem>* CachedUncontrolledFileItem = UncontrolledChangelistItemCache.Find(FileState); CachedUncontrolledFileItem && CachedUncontrolledFileItem->IsValid())
 			{
 				check((*CachedUncontrolledFileItem)->GetTreeItemType() == IChangelistTreeItem::File);
 				UncontrolledFileItem = StaticCastSharedPtr<FFileTreeItem>(*CachedUncontrolledFileItem);
