Fix issues with icon color/padding in the changelists window

diff --git a/Engine/Source/Editor/SourceControlWindows/Private/SSourceControlChangelistRows.cpp b/Engine/Source/Editor/SourceControlWindows/Private/SSourceControlChangelistRows.cpp
--- a/Engine/Source/Editor/SourceControlWindows/Private/SSourceControlChangelistRows.cpp
+++ b/Engine/Source/Editor/SourceControlWindows/Private/SSourceControlChangelistRows.cpp
@@ -70,10 +70,6 @@ void SChangelistTableRow::Construct(const FArguments& InArgs, const TSharedRef<S
 	TreeItem = static_cast<FChangelistTreeItem*>(InArgs._TreeItemToVisualize.Get());
 	OnPostDrop = InArgs._OnPostDrop;
 
-	const FSlateBrush* IconBrush = (TreeItem != nullptr) ?
-		FAppStyle::GetBrush(TreeItem->ChangelistState->GetSmallIconName()) :
-		FAppStyle::GetBrush("SourceControl.Changelist");
-
 	SetToolTipText(GetChangelistDescriptionText());
 
 	STableRow<FChangelistTreeItemPtr>::Construct(
@@ -86,7 +82,10 @@ void SChangelistTableRow::Construct(const FArguments& InArgs, const TSharedRef<S
 			.AutoWidth()
 			[
 				SNew(SImage)
-				.Image(IconBrush)
+				.Image_Lambda([this]()
+				{
+					return (TreeItem != nullptr) ?FAppStyle::GetBrush(TreeItem->ChangelistState->GetSmallIconName()) : FAppStyle::GetBrush("SourceControl.Changelist");
+				})
 			]
 			+SHorizontalBox::Slot() // Changelist number.
 			.Padding(2, 0, 0, 0)
diff --git a/Engine/Source/Editor/SourceControlWindows/Private/SSourceControlCommon.cpp b/Engine/Source/Editor/SourceControlWindows/Private/SSourceControlCommon.cpp
--- a/Engine/Source/Editor/SourceControlWindows/Private/SSourceControlCommon.cpp
+++ b/Engine/Source/Editor/SourceControlWindows/Private/SSourceControlCommon.cpp
@@ -396,6 +396,12 @@ FText GetSingleLineChangelistDescription(const FText& InFullDescription)
 		DescriptionTextAsString.LeftInline(NewlineStartIndex);
 	}
 
+	// Trim any trailing carriage returns
+	if (DescriptionTextAsString.EndsWith(TEXT("\r"), ESearchCase::CaseSensitive))
+	{
+		DescriptionTextAsString.LeftChopInline(1);
+	}
+
 	return InFullDescription.IsCultureInvariant() ? FText::AsCultureInvariant(DescriptionTextAsString) : FText::FromString(DescriptionTextAsString);
 }
 
