Converting Revision Control History window to use UToolMenu for its context menu

diff --git a/Engine/Source/Editor/SourceControlWindows/Private/SSourceControlHistory.cpp b/Engine/Source/Editor/SourceControlWindows/Private/SSourceControlHistory.cpp
--- a/Engine/Source/Editor/SourceControlWindows/Private/SSourceControlHistory.cpp
+++ b/Engine/Source/Editor/SourceControlWindows/Private/SSourceControlHistory.cpp
@@ -50,6 +50,8 @@
 #include "IAssetTools.h"
 #include "IAssetTypeActions.h"
 #include "AssetToolsModule.h"
+#include "ToolMenu.h"
+#include "ToolMenus.h"
 
 /**
  * Wrapper around data from ISourceControlRevision
@@ -643,6 +645,7 @@ public:
 
 	SLATE_END_ARGS()
 		
+	const FName SourceControlHistoryContextMenu = TEXT("RevisionControl.History.ContextMenu");
 
 	SSourceControlHistoryWidget()
 	{
@@ -650,6 +653,53 @@ public:
 
 	void Construct( const FArguments& InArgs )
 	{	
+		if (!UToolMenus::Get()->IsMenuRegistered(SourceControlHistoryContextMenu))
+		{
+			UToolMenu* ContextMenu = UToolMenus::Get()->RegisterMenu(SourceControlHistoryContextMenu);
+			ContextMenu->bShouldCloseWindowAfterMenuSelection = true;
+			
+			FToolMenuSection& DiffSection = ContextMenu->AddSection("DiffTools");
+			DiffSection.AddMenuEntry(
+				TEXT("DiffAgainstPrevious"),
+				NSLOCTEXT("SourceControl.HistoryWindow.Menu", "DiffAgainstPrev", "Diff Against Previous Revision"),
+				NSLOCTEXT("SourceControl.HistoryWindow.Menu", "DiffAgainstPrevTooltip", "See changes between this revision and the previous one."),
+				FSlateIcon(),
+				FUIAction(
+					FExecuteAction::CreateSP(this, &SSourceControlHistoryWidget::OnDiffAgainstPreviousRev),
+					FCanExecuteAction::CreateSP(this, &SSourceControlHistoryWidget::CanDiffAgainstPreviousRev)
+				)
+			);
+
+			DiffSection.AddMenuEntry(
+				TEXT("DiffAgainstWorkspace"),
+				NSLOCTEXT("SourceControl.HistoryWindow.Menu", "DiffAgainstWorkspace", "Diff Against Workspace File"),
+				NSLOCTEXT("SourceControl.HistoryWindow.Menu", "DiffAgainstWorkspaceTooltip", "See changes between this revision and your version of the asset."),
+				FSlateIcon(),
+				FUIAction(
+					FExecuteAction::CreateSP(this, &SSourceControlHistoryWidget::OnDiffAgainstWorkspace),
+					FCanExecuteAction::CreateSP(this, &SSourceControlHistoryWidget::CanDiffAgainstWorkspace)
+				)
+			);
+
+		
+			DiffSection.AddDynamicEntry(NAME_None, FNewToolMenuSectionDelegate::CreateLambda([this](FToolMenuSection& InSection)
+			{
+				if (this->CanDiffSelected())
+				{
+					InSection.AddMenuEntry(
+						TEXT("DiffSelected"),
+						NSLOCTEXT("SourceControl.HistoryWindow.Menu", "DiffSelected", "Diff Selected"),
+						NSLOCTEXT("SourceControl.HistoryWindow.Menu", "DiffSelectedTooltip", "Diff the two assets that you have selected."),
+						FSlateIcon(),
+						FUIAction(
+							FExecuteAction::CreateSP(this, &SSourceControlHistoryWidget::OnDiffSelected)
+						)
+					);
+				}
+			}));
+			
+		}
+
 		AddHistoryInfo(InArgs._SourceControlStates.Get());
 		ParentWindow = InArgs._ParentWindow.Get();
 
@@ -1152,46 +1202,13 @@ private:
 	{
 		if (CanDiff())
 		{
-			FMenuBuilder MenuBuilder(true, NULL);
-
-			MenuBuilder.AddMenuEntry(
-				NSLOCTEXT("SourceControl.HistoryWindow.Menu", "DiffAgainstPrev", "Diff Against Previous Revision"),
-				NSLOCTEXT("SourceControl.HistoryWindow.Menu", "DiffAgainstPrevTooltip", "See changes between this revision and the previous one."),
-				FSlateIcon(),
-				FUIAction(
-					FExecuteAction::CreateSP(this, &SSourceControlHistoryWidget::OnDiffAgainstPreviousRev),
-					FCanExecuteAction::CreateSP(this, &SSourceControlHistoryWidget::CanDiffAgainstPreviousRev)
-				)
-			);
-
-			MenuBuilder.AddMenuEntry(
-				NSLOCTEXT("SourceControl.HistoryWindow.Menu", "DiffAgainstWorkspace", "Diff Against Workspace File"),
-				NSLOCTEXT("SourceControl.HistoryWindow.Menu", "DiffAgainstWorkspaceTooltip", "See changes between this revision and your version of the asset."),
-				FSlateIcon(),
-				FUIAction(
-					FExecuteAction::CreateSP(this, &SSourceControlHistoryWidget::OnDiffAgainstWorkspace),
-					FCanExecuteAction::CreateSP(this, &SSourceControlHistoryWidget::CanDiffAgainstWorkspace)
-				)
-			);
-
-			if (CanDiffSelected())
+			FToolMenuContext Context;
+			if (UToolMenu* GeneratedContextMenu = UToolMenus::Get()->GenerateMenu(SourceControlHistoryContextMenu, Context))
 			{
-				MenuBuilder.AddMenuEntry(
-					NSLOCTEXT("SourceControl.HistoryWindow.Menu", "DiffSelected", "Diff Selected"),
-					NSLOCTEXT("SourceControl.HistoryWindow.Menu", "DiffSelectedTooltip", "Diff the two assets that you have selected."),
-					FSlateIcon(),
-					FUIAction(
-						FExecuteAction::CreateSP(this, &SSourceControlHistoryWidget::OnDiffSelected)
-					)
-				);
+				return UToolMenus::Get()->GenerateWidget(GeneratedContextMenu);
 			}
-
-			return MenuBuilder.MakeWidget();
-		}
-		else
-		{
-			return SNullWidget::NullWidget;
 		}
+		return SNullWidget::NullWidget;
 	}
 
 	/** See if we should enabled the 'diff against previous' option */
