Added context menu option to revert assets in the submit changes dialog

diff --git a/Engine/Source/Editor/SourceControlWindows/Private/SSourceControlCommon.h b/Engine/Source/Editor/SourceControlWindows/Private/SSourceControlCommon.h
--- a/Engine/Source/Editor/SourceControlWindows/Private/SSourceControlCommon.h
+++ b/Engine/Source/Editor/SourceControlWindows/Private/SSourceControlCommon.h
@@ -279,6 +279,9 @@ struct FFileTreeItem : public IFileViewTreeItem
 	/** true if the item is source controlled and not marked for add nor for delete */
 	bool CanDiff() const { return FileState->IsSourceControlled() && !FileState->IsAdded() && !FileState->IsDeleted(); }
 
+	/** true if the item is source controlled and can be reverted */
+	bool CanRevert() const { return FileState->IsSourceControlled() && FileState->CanRevert(); }
+
 	const FAssetDataArrayPtr& GetAssetData() const
 	{
 		return Assets;
diff --git a/Engine/Source/Editor/SourceControlWindows/Private/SSourceControlSubmit.cpp b/Engine/Source/Editor/SourceControlWindows/Private/SSourceControlSubmit.cpp
--- a/Engine/Source/Editor/SourceControlWindows/Private/SSourceControlSubmit.cpp
+++ b/Engine/Source/Editor/SourceControlWindows/Private/SSourceControlSubmit.cpp
@@ -27,6 +27,8 @@
 #include "Virtualization/VirtualizationSystem.h"
 #include "Logging/MessageLog.h"
 #include "RevisionControlStyle/RevisionControlStyle.h"
+#include "Bookmarks/BookmarkScoped.h"
+#include "Algo/AllOf.h"
 
 #if SOURCE_CONTROL_WITH_SLATE
 
@@ -248,7 +250,7 @@ void SSourceControlSubmitWidget::Construct(const FArguments& InArgs)
 			.OnContextMenuOpening(this, &SSourceControlSubmitWidget::OnCreateContextMenu)
 			.OnMouseButtonDoubleClick(this, &SSourceControlSubmitWidget::OnDiffAgainstDepotSelected)
 			.HeaderRow(HeaderRowWidget)
-			.SelectionMode(ESelectionMode::Single)
+			.SelectionMode(ESelectionMode::Multi)
 		]
 	];
 
@@ -437,11 +439,11 @@ void SSourceControlSubmitWidget::Construct(const FArguments& InArgs)
 /** Corvus: Called to create a context menu when right-clicking on an item */
 TSharedPtr<SWidget> SSourceControlSubmitWidget::OnCreateContextMenu()
 {
-	if (SSourceControlSubmitWidget::CanDiffAgainstDepot())
-	{
-		FMenuBuilder MenuBuilder(true, NULL);
+	FMenuBuilder MenuBuilder(true, NULL);
 
-		MenuBuilder.BeginSection("Source Control", NSLOCTEXT("SourceControl.SubmitWindow.Menu", "SourceControlSectionHeader", "Revision Control"));
+	MenuBuilder.BeginSection("Source Control", NSLOCTEXT("SourceControl.SubmitWindow.Menu", "SourceControlSectionHeader", "Revision Control"));
+	{
+		if (SSourceControlSubmitWidget::CanDiffAgainstDepot())
 		{
 			MenuBuilder.AddMenuEntry(
 				NSLOCTEXT("SourceControl.SubmitWindow.Menu", "DiffAgainstDepot", "Diff Against Depot"),
@@ -453,14 +455,20 @@ TSharedPtr<SWidget> SSourceControlSubmitWidget::OnCreateContextMenu()
 				)
 			);
 		}
-		MenuBuilder.EndSection();
 
-		return MenuBuilder.MakeWidget();
-	}
-	else
-	{
-		return nullptr;
+		MenuBuilder.AddMenuEntry(
+			NSLOCTEXT("SourceControl.SubmitWindow.Menu", "Revert", "Revert"),
+			NSLOCTEXT("SourceControl.SubmitWindow.Menu", "RevertTooltip", "Revert the selected assets to their original state from revision control."),
+			FSlateIcon(FAppStyle::GetAppStyleSetName(), "SourceControl.Actions.Revert"),
+			FUIAction(
+				FExecuteAction::CreateSP(this, &SSourceControlSubmitWidget::OnRevert), 
+				FCanExecuteAction::CreateSP(this, &SSourceControlSubmitWidget::CanRevert)
+			)
+		);
 	}
+	MenuBuilder.EndSection();
+
+	return MenuBuilder.MakeWidget();
 }
 
 bool SSourceControlSubmitWidget::CanDiffAgainstDepot() const
@@ -511,6 +519,77 @@ void SSourceControlSubmitWidget::OnDiffAgainstDepotSelected(TSharedPtr<FFileTree
 	}
 }
 
+bool SSourceControlSubmitWidget::CanRevert() const
+{
+	const auto& SelectedItems = ListView->GetSelectedItems();
+	if (SelectedItems.Num() > 0)
+	{
+		return Algo::AllOf(SelectedItems, [](const FFileTreeItemPtr& SelectedItem)
+			{
+				return SelectedItem->CanRevert();
+			}
+		);
+	}
+	return false;
+}
+
+void SSourceControlSubmitWidget::OnRevert()
+{
+	const auto& SelectedItems = ListView->GetSelectedItems();
+	if (SelectedItems.Num() < 1)
+	{
+		return;
+	}
+
+	auto RemoveItemsFromListView = [this](TArray<FString>& ItemsToRemove)
+	{
+		ListViewItems.RemoveAll([&ItemsToRemove](const FFileTreeItemPtr& ListViewItem) -> bool
+			{
+				return ItemsToRemove.ContainsByPredicate([&ListViewItem](const FString& ItemToRemove) -> bool
+					{
+						return ItemToRemove == ListViewItem->GetFileName().ToString();
+					}
+				);
+			}
+		);
+	};
+
+	TArray<FString> PackagesToRevert;
+	TArray<FString> FilesToRevert;
+	for (const auto& SelectedItem : SelectedItems)
+	{
+		FString PackageName;
+		if (FPackageName::TryConvertFilenameToLongPackageName(SelectedItem->GetFileName().ToString(), PackageName))
+		{
+			PackagesToRevert.Add(SelectedItem->GetFileName().ToString());
+		}
+		else
+		{
+			FilesToRevert.Add(SelectedItem->GetFileName().ToString());
+		}
+	}
+
+	{
+		FBookmarkScoped BookmarkScoped;
+		bool bAnyReverted = false;
+		if (PackagesToRevert.Num() > 0)
+		{
+			bAnyReverted = SourceControlHelpers::RevertAndReloadPackages(PackagesToRevert, /*bRevertAll=*/false, /*bReloadWorld=*/true);
+			RemoveItemsFromListView(PackagesToRevert);
+		}
+		if (FilesToRevert.Num() > 0)
+		{
+			bAnyReverted |= SourceControlHelpers::RevertFiles(FilesToRevert);
+			RemoveItemsFromListView(FilesToRevert);
+		}
+		
+		if (bAnyReverted)
+		{
+			ListView->RebuildList();
+		}
+	}
+}
+
 FReply SSourceControlSubmitWidget::OnKeyDown( const FGeometry& MyGeometry, const FKeyEvent& InKeyEvent )
 {
 	// Pressing escape returns as if the user clicked cancel
@@ -671,7 +750,7 @@ FReply SSourceControlSubmitWidget::SaveAndCloseClicked()
 
 bool SSourceControlSubmitWidget::IsSubmitEnabled() const
 {
-	return bAllowSubmit && !ChangeListDescriptionTextCtrl->GetText().IsEmpty();
+	return bAllowSubmit && !ChangeListDescriptionTextCtrl->GetText().IsEmpty() && ListViewItems.Num() > 0;
 }
 
 
diff --git a/Engine/Source/Editor/SourceControlWindows/Private/SSourceControlSubmit.h b/Engine/Source/Editor/SourceControlWindows/Private/SSourceControlSubmit.h
--- a/Engine/Source/Editor/SourceControlWindows/Private/SSourceControlSubmit.h
+++ b/Engine/Source/Editor/SourceControlWindows/Private/SSourceControlSubmit.h
@@ -176,6 +176,12 @@ private:
 	void OnDiffAgainstDepot();
 	void OnDiffAgainstDepotSelected(TSharedPtr<FFileTreeItem> InSelectedItem);
 
+	/** Called to check whether the selected files in the ListView can be reverted */
+	bool CanRevert() const;
+
+	/** Reverts the files selected in the ListView */
+	void OnRevert();
+
 private:
 	ESubmitResults::Type DialogResult;
 
