Change to capture weak pointer to the revision control UI in the dynamic menu entry.

diff --git a/Engine/Source/Editor/SourceControlWindows/Private/SSourceControlHistory.cpp b/Engine/Source/Editor/SourceControlWindows/Private/SSourceControlHistory.cpp
--- a/Engine/Source/Editor/SourceControlWindows/Private/SSourceControlHistory.cpp
+++ b/Engine/Source/Editor/SourceControlWindows/Private/SSourceControlHistory.cpp
@@ -681,22 +681,7 @@ public:
 				)
 			);
 
-		
-			DiffSection.AddDynamicEntry(NAME_None, FNewToolMenuSectionDelegate::CreateLambda([this](FToolMenuSection& InSection)
-			{
-				if (this->CanDiffSelected())
-				{
-					InSection.AddMenuEntry(
-						TEXT("DiffSelected"),
-						NSLOCTEXT("SourceControl.HistoryWindow.Menu", "DiffSelected", "Diff Selected"),
-						NSLOCTEXT("SourceControl.HistoryWindow.Menu", "DiffSelectedTooltip", "Diff the two assets that you have selected."),
-						FSlateIcon(),
-						FUIAction(
-							FExecuteAction::CreateSP(this, &SSourceControlHistoryWidget::OnDiffSelected)
-						)
-					);
-				}
-			}));
+			DiffSection.AddDynamicEntry(NAME_None, FNewToolMenuSectionDelegate::CreateSP(this, &SSourceControlHistoryWidget::CreateDiffSelectedMenu));
 			
 		}
 
@@ -1469,6 +1454,25 @@ private:
 		}
 	}
 
+	/** 
+	* CanDiffSelected gives a dynamic response based on a number of factors, so CreateDiffSelectedMenu ensures we are checking on each menu open event and correctly adding the related menu item(s) for each case.
+	*/
+	void CreateDiffSelectedMenu(FToolMenuSection& InSection)
+	{
+		if (CanDiffSelected())
+		{
+			InSection.AddMenuEntry(
+				TEXT("DiffSelected"),
+				NSLOCTEXT("SourceControl.HistoryWindow.Menu", "DiffSelected", "Diff Selected"),
+				NSLOCTEXT("SourceControl.HistoryWindow.Menu", "DiffSelectedTooltip", "Diff the two assets that you have selected."),
+				FSlateIcon(),
+				FUIAction(
+					FExecuteAction::CreateSP(this, &SSourceControlHistoryWidget::OnDiffSelected)
+				)
+			);
+		}
+	}
+
 	/**
 	 * An event handler for mouse drag detection events. Intended to be used as a delegate for
 	 * history tree rows. Creates a FSourceControlHistoryRowDragDropOp (if the drag was with 
