Disallow unsaved changes to be discarded from unsaved levels.

This previously resulted in the level being deleted irreversibly.

diff --git a/Engine/Source/Editor/SourceControlWindows/Private/SSourceControlChangelistRows.cpp b/Engine/Source/Editor/SourceControlWindows/Private/SSourceControlChangelistRows.cpp
--- a/Engine/Source/Editor/SourceControlWindows/Private/SSourceControlChangelistRows.cpp
+++ b/Engine/Source/Editor/SourceControlWindows/Private/SSourceControlChangelistRows.cpp
@@ -17,6 +17,7 @@
 #include "RevisionControlStyle/RevisionControlStyle.h"
 #include "Styling/SlateStyleRegistry.h"
 #include "Styling/StarshipCoreStyle.h"
+#include "Misc/MessageDialog.h"
 #include "Widgets/Images/SThrobber.h"
 #include "Widgets/Layout/SWidgetSwitcher.h"
 
@@ -549,6 +550,26 @@ TSharedRef<SWidget> SOfflineFileTableRow::GenerateWidgetForColumn(const FName& C
 				.ColorAndOpacity(FSlateColor::UseSubduedForeground())
 				.OnMouseButtonDown_Lambda([this, DiscardSwitcher] (const FGeometry&, const FPointerEvent&) -> FReply
 				{
+					// Normalize packagenames and filenames
+					FString PackageName;
+					{
+						FString TreeName = TreeItem->GetPackageName().ToString();
+
+						if (!FPackageName::TryConvertFilenameToLongPackageName(TreeName, PackageName))
+						{
+							PackageName = MoveTemp(TreeName);
+						}
+					}
+					// Validate we have a saved map
+					UPackage* LevelPackage = FindPackage(nullptr, *PackageName)->GetOutermost();
+					if (LevelPackage == GetTransientPackage()
+						|| LevelPackage->HasAnyFlags(RF_Transient)
+						|| !FPackageName::IsValidLongPackageName(LevelPackage->GetName()))
+					{
+						FMessageDialog::Open(EAppMsgType::Ok, LOCTEXT("DiscardUnsavedChangesSaveMap", "You need to save the level before discarding unsaved changes."));
+						return FReply::Handled();
+					}
+					
 					DiscardSwitcher->SetActiveWidgetIndex(1);
 					Async(EAsyncExecution::TaskGraphMainThread,
 						[this]
