Auto select items in revert window to reduce redundant actions for users.

#tests used menu to revert assets

diff --git a/Engine/Source/Editor/SourceControlWindows/Private/SSourceControlRevert.cpp b/Engine/Source/Editor/SourceControlWindows/Private/SSourceControlRevert.cpp
--- a/Engine/Source/Editor/SourceControlWindows/Private/SSourceControlRevert.cpp
+++ b/Engine/Source/Editor/SourceControlWindows/Private/SSourceControlRevert.cpp
@@ -59,7 +59,7 @@ struct FRevertCheckBoxListViewItem
 	FRevertCheckBoxListViewItem( FString InText )
 	{
 		Text = InText;
-		IsSelected = false;
+		IsSelected = true;
 		IsModified = false;
 	}
 
