Create revision control context object for menu creation

diff --git a/Engine/Source/Editor/SourceControlWindows/Private/SSourceControlHistory.cpp b/Engine/Source/Editor/SourceControlWindows/Private/SSourceControlHistory.cpp
--- a/Engine/Source/Editor/SourceControlWindows/Private/SSourceControlHistory.cpp
+++ b/Engine/Source/Editor/SourceControlWindows/Private/SSourceControlHistory.cpp
@@ -651,38 +651,16 @@ public:
 	{
 	}
 
+	static void CreateDiffMenu(UToolMenu* InToolMenu);
+
 	void Construct( const FArguments& InArgs )
 	{	
 		if (!UToolMenus::Get()->IsMenuRegistered(SourceControlHistoryContextMenu))
 		{
 			UToolMenu* ContextMenu = UToolMenus::Get()->RegisterMenu(SourceControlHistoryContextMenu);
 			ContextMenu->bShouldCloseWindowAfterMenuSelection = true;
-			
-			FToolMenuSection& DiffSection = ContextMenu->AddSection("DiffTools");
-			DiffSection.AddMenuEntry(
-				TEXT("DiffAgainstPrevious"),
-				NSLOCTEXT("SourceControl.HistoryWindow.Menu", "DiffAgainstPrev", "Diff Against Previous Revision"),
-				NSLOCTEXT("SourceControl.HistoryWindow.Menu", "DiffAgainstPrevTooltip", "See changes between this revision and the previous one."),
-				FSlateIcon(),
-				FUIAction(
-					FExecuteAction::CreateSP(this, &SSourceControlHistoryWidget::OnDiffAgainstPreviousRev),
-					FCanExecuteAction::CreateSP(this, &SSourceControlHistoryWidget::CanDiffAgainstPreviousRev)
-				)
-			);
-
-			DiffSection.AddMenuEntry(
-				TEXT("DiffAgainstWorkspace"),
-				NSLOCTEXT("SourceControl.HistoryWindow.Menu", "DiffAgainstWorkspace", "Diff Against Workspace File"),
-				NSLOCTEXT("SourceControl.HistoryWindow.Menu", "DiffAgainstWorkspaceTooltip", "See changes between this revision and your version of the asset."),
-				FSlateIcon(),
-				FUIAction(
-					FExecuteAction::CreateSP(this, &SSourceControlHistoryWidget::OnDiffAgainstWorkspace),
-					FCanExecuteAction::CreateSP(this, &SSourceControlHistoryWidget::CanDiffAgainstWorkspace)
-				)
-			);
 
-			DiffSection.AddDynamicEntry(NAME_None, FNewToolMenuSectionDelegate::CreateSP(this, &SSourceControlHistoryWidget::CreateDiffSelectedMenu));
-			
+			ContextMenu->AddDynamicSection(NAME_None, FNewToolMenuDelegate::CreateStatic(&SSourceControlHistoryWidget::CreateDiffMenu));
 		}
 
 		AddHistoryInfo(InArgs._SourceControlStates.Get());
@@ -1188,6 +1166,9 @@ private:
 		if (CanDiff())
 		{
 			FToolMenuContext Context;
+			USourceControlHistoryWidgetContext* SourceControlHistoryWidgetContext = NewObject<USourceControlHistoryWidgetContext>();
+			SourceControlHistoryWidgetContext->HistoryWidget = SharedThis(this);
+			Context.AddObject(SourceControlHistoryWidgetContext);
 			if (UToolMenu* GeneratedContextMenu = UToolMenus::Get()->GenerateMenu(SourceControlHistoryContextMenu, Context))
 			{
 				return UToolMenus::Get()->GenerateWidget(GeneratedContextMenu);
@@ -1626,6 +1607,46 @@ private:
 	TWeakPtr<SWindow> ParentWindow;
 };
 
+void SSourceControlHistoryWidget::CreateDiffMenu(UToolMenu* InToolMenu)
+{
+	USourceControlHistoryWidgetContext* FoundContext = InToolMenu->FindContext<USourceControlHistoryWidgetContext>();
+	if (!FoundContext)
+	{
+		return;
+	}
+
+	TSharedPtr<SSourceControlHistoryWidget> PinnedHistoryWidget = FoundContext->HistoryWidget.Pin();
+	if (!PinnedHistoryWidget)
+	{
+		return;
+	}
+
+	FToolMenuSection& DiffSection = InToolMenu->AddSection("DiffTools");
+	DiffSection.AddMenuEntry(
+		TEXT("DiffAgainstPrevious"),
+		NSLOCTEXT("SourceControl.HistoryWindow.Menu", "DiffAgainstPrev", "Diff Against Previous Revision"),
+		NSLOCTEXT("SourceControl.HistoryWindow.Menu", "DiffAgainstPrevTooltip", "See changes between this revision and the previous one."),
+		FSlateIcon(),
+		FUIAction(
+			FExecuteAction::CreateSP(PinnedHistoryWidget.Get(), &SSourceControlHistoryWidget::OnDiffAgainstPreviousRev),
+			FCanExecuteAction::CreateSP(PinnedHistoryWidget.Get(), &SSourceControlHistoryWidget::CanDiffAgainstPreviousRev)
+		)
+	);
+
+	DiffSection.AddMenuEntry(
+		TEXT("DiffAgainstWorkspace"),
+		NSLOCTEXT("SourceControl.HistoryWindow.Menu", "DiffAgainstWorkspace", "Diff Against Workspace File"),
+		NSLOCTEXT("SourceControl.HistoryWindow.Menu", "DiffAgainstWorkspaceTooltip", "See changes between this revision and your version of the asset."),
+		FSlateIcon(),
+		FUIAction(
+			FExecuteAction::CreateSP(PinnedHistoryWidget.Get(), &SSourceControlHistoryWidget::OnDiffAgainstWorkspace),
+			FCanExecuteAction::CreateSP(PinnedHistoryWidget.Get(), &SSourceControlHistoryWidget::CanDiffAgainstWorkspace)
+		)
+	);
+
+	DiffSection.AddDynamicEntry(NAME_None, FNewToolMenuSectionDelegate::CreateSP(PinnedHistoryWidget.Get(), &SSourceControlHistoryWidget::CreateDiffSelectedMenu));
+}
+
 void FSourceControlWindows::DisplayRevisionHistory( const TArray<FString>& InPackageNames )
 {
 	// Explicitly load the module so live coding will work with it
diff --git a/Engine/Source/Editor/SourceControlWindows/Public/SourceControlWindows.h b/Engine/Source/Editor/SourceControlWindows/Public/SourceControlWindows.h
--- a/Engine/Source/Editor/SourceControlWindows/Public/SourceControlWindows.h
+++ b/Engine/Source/Editor/SourceControlWindows/Public/SourceControlWindows.h
@@ -5,6 +5,22 @@
 #include "CoreMinimal.h"
 #include "SlateFwd.h"
 #include "ISourceControlProvider.h"
+#include "Templates/SharedPointer.h"
+#include "UObject/Object.h"
+#include "UObject/ObjectMacros.h"
+#include "SourceControlWindows.generated.h"
+
+class SSourceControlHistoryWidget;
+
+/** Context object for SSourceControlHistoryWidget's right-click context menu to get info about which history widget is trying to generate the menu */
+UCLASS()
+class USourceControlHistoryWidgetContext : public UObject
+{
+	GENERATED_BODY()
+
+public:
+	TWeakPtr<SSourceControlHistoryWidget> HistoryWidget;
+};
 
 
 /** Info supplied as argument to delegate FSourceControlWindowsOnCheckInComplete called by FSourceControlWindows::ChoosePackagesToCheckIn() and optional argument to PromptForCheckin(). */
