[CrashFix] Fix Dereference Null
#pf none

diff --git a/Engine/Source/Editor/SourceControlWindows/Private/SSourceControlHistory.cpp b/Engine/Source/Editor/SourceControlWindows/Private/SSourceControlHistory.cpp
--- a/Engine/Source/Editor/SourceControlWindows/Private/SSourceControlHistory.cpp
+++ b/Engine/Source/Editor/SourceControlWindows/Private/SSourceControlHistory.cpp
@@ -230,7 +230,7 @@ static UObject* GetAssetRevisionObject(TSharedPtr<FHistoryTreeItem> HistoryTreeI
 		// grab the asset from the package - we assume asset name matches file name
 		FString AssetName = FPaths::GetBaseFilename(FileListItem->FileName);
 		AssetObject = FindObject<UObject>(AssetPackage, *AssetName);
-		if (AssetObject == nullptr)
+		if (AssetPackage && !AssetObject)
 		{
 			AssetObject = AssetPackage->FindAssetInPackage();
 		}
