Add support for cherry picking an asset from the asset history dialog, using Tool Menus

diff --git a/Engine/Source/Editor/SourceControlWindows/Private/SSourceControlHistory.cpp b/Engine/Source/Editor/SourceControlWindows/Private/SSourceControlHistory.cpp
--- a/Engine/Source/Editor/SourceControlWindows/Private/SSourceControlHistory.cpp
+++ b/Engine/Source/Editor/SourceControlWindows/Private/SSourceControlHistory.cpp
@@ -1168,17 +1168,30 @@ private:
 	/** Called to create a context menu when right-clicking on a history item */
 	TSharedPtr< SWidget > OnCreateContextMenu()
 	{
-		if (CanDiff())
+		FToolMenuContext Context;
+		USourceControlHistoryWidgetContext* SourceControlHistoryWidgetContext = NewObject<USourceControlHistoryWidgetContext>();
+		SourceControlHistoryWidgetContext->HistoryWidget = SharedThis(this);
+
+		for (TSharedPtr<FHistoryTreeItem> Item : MainHistoryListView->GetSelectedItems())
 		{
-			FToolMenuContext Context;
-			USourceControlHistoryWidgetContext* SourceControlHistoryWidgetContext = NewObject<USourceControlHistoryWidgetContext>();
-			SourceControlHistoryWidgetContext->HistoryWidget = SharedThis(this);
-			Context.AddObject(SourceControlHistoryWidgetContext);
-			if (UToolMenu* GeneratedContextMenu = UToolMenus::Get()->GenerateMenu(SourceControlHistoryContextMenu, Context))
+			if (Item->Parent.IsValid())
 			{
-				return UToolMenus::Get()->GenerateWidget(GeneratedContextMenu);
+				TSharedPtr<FHistoryTreeItem> Parent = Item->Parent.Pin();
+
+				USourceControlHistoryWidgetContext::SelectedItem ItemInfo;
+				ItemInfo.FileName = Parent->FileListItem->FileName;
+				ItemInfo.Revision = Item->RevisionListItem->Revision;
+
+				SourceControlHistoryWidgetContext->GetSelectedItems().Add(ItemInfo);
 			}
 		}
+	
+		Context.AddObject(SourceControlHistoryWidgetContext);
+		if (UToolMenu* GeneratedContextMenu = UToolMenus::Get()->GenerateMenu(SourceControlHistoryContextMenu, Context))
+		{
+			return UToolMenus::Get()->GenerateWidget(GeneratedContextMenu);
+		}
+
 		return SNullWidget::NullWidget;
 	}
 
@@ -1405,7 +1418,7 @@ private:
 	 *
 	 * @return True if the SourceControl provider supports diffing, false if not.
 	 */
-	bool CanDiff() const
+	static bool CanDiff()
 	{
 		return ISourceControlModule::Get().GetProvider().AllowsDiffAgainstDepot();
 	}
@@ -1626,6 +1639,11 @@ void SSourceControlHistoryWidget::CreateDiffMenu(UToolMenu* InToolMenu)
 		return;
 	}
 
+	if (!CanDiff())
+	{
+		return;
+	}
+
 	FToolMenuSection& DiffSection = InToolMenu->AddSection("DiffTools");
 	DiffSection.AddMenuEntry(
 		TEXT("DiffAgainstPrevious"),
diff --git a/Engine/Source/Editor/SourceControlWindows/Public/SourceControlWindows.h b/Engine/Source/Editor/SourceControlWindows/Public/SourceControlWindows.h
--- a/Engine/Source/Editor/SourceControlWindows/Public/SourceControlWindows.h
+++ b/Engine/Source/Editor/SourceControlWindows/Public/SourceControlWindows.h
@@ -14,12 +14,26 @@ class SSourceControlHistoryWidget;
 
 /** Context object for SSourceControlHistoryWidget's right-click context menu to get info about which history widget is trying to generate the menu */
 UCLASS()
-class USourceControlHistoryWidgetContext : public UObject
+class SOURCECONTROLWINDOWS_API USourceControlHistoryWidgetContext : public UObject
 {
 	GENERATED_BODY()
 
 public:
-	TWeakPtr<SSourceControlHistoryWidget> HistoryWidget;
+	struct SelectedItem
+	{
+		FString FileName;
+		FString Revision;
+	};
+
+	TWeakPtr<class SSourceControlHistoryWidget> HistoryWidget;
+
+	TArray<SelectedItem>& GetSelectedItems()
+	{
+		return SelectedItems;
+	}
+
+private:
+	TArray<SelectedItem> SelectedItems;
 };
 
 
