SourceControl: After reverting all files in the 'Submit Files' dialog, close it.

diff --git a/Engine/Source/Editor/SourceControlWindows/Private/SSourceControlSubmit.cpp b/Engine/Source/Editor/SourceControlWindows/Private/SSourceControlSubmit.cpp
--- a/Engine/Source/Editor/SourceControlWindows/Private/SSourceControlSubmit.cpp
+++ b/Engine/Source/Editor/SourceControlWindows/Private/SSourceControlSubmit.cpp
@@ -601,7 +601,15 @@ void SSourceControlSubmitWidget::OnRevert()
 		
 		if (bAnyReverted)
 		{
-			ListView->RebuildList();
+			if (ListViewItems.IsEmpty())
+			{
+				DialogResult = ESubmitResults::SUBMIT_CANCELED;
+				ParentFrame.Pin()->RequestDestroyWindow();
+			}
+			else
+			{
+				ListView->RebuildList();
+			}
 		}
 	}
 }
