Fixed missing '#changelist validated' tag when submitting without editing the changelist text.

diff --git a/Engine/Source/Editor/SourceControlWindows/Private/SSourceControlChangelists.cpp b/Engine/Source/Editor/SourceControlWindows/Private/SSourceControlChangelists.cpp
--- a/Engine/Source/Editor/SourceControlWindows/Private/SSourceControlChangelists.cpp
+++ b/Engine/Source/Editor/SourceControlWindows/Private/SSourceControlChangelists.cpp
@@ -2396,11 +2396,15 @@ void SSourceControlChangelistsWidget::OnSubmitChangelist()
 		// Check if any of the presubmit hooks fail. (This might also update the changelist description)
 		if (GetOnPresubmitResult(ChangelistState, Description))
 		{
-			// If the description was modified, add it to the operation to update the changelist
+			// If the user modified the description, ensure the 'validation tag' wasn't removed if validation is enabled.
 			if (!ChangelistDescriptionToSubmit.EqualTo(Description.Description))
 			{
 				SubmitChangelistOperation->SetDescription(UpdateChangelistDescriptionToSubmitIfNeeded(bValidationResult, Description.Description));
 			}
+			else
+			{
+				SubmitChangelistOperation->SetDescription(Description.Description);
+			}
 
 			Execute(LOCTEXT("Submitting_Changelist", "Submitting changelist..."), SubmitChangelistOperation, ChangelistState->GetChangelist(), EConcurrency::Synchronous, FSourceControlOperationComplete::CreateLambda(
 				[&SubmitChangelistOperation, &bCheckinSuccess](const TSharedRef<ISourceControlOperation>& Operation, ECommandResult::Type InResult)
