SourceControl: Refactored 'revert unsaved' behavior to be consistent in SceneOutliner and ContentBrowser, following #23845005.

- Can be enabled/disabled via cvar (default disabled, enabled in Skein).
- Can revert unsaved in ContentBrowser (=new)
- Can revert unsaved in SceneOutliner (=existing)

diff --git a/Engine/Plugins/Editor/AssetManagerEditor/Source/AssetManagerEditor/Private/AssetSourceControlContextMenu.cpp b/Engine/Plugins/Editor/AssetManagerEditor/Source/AssetManagerEditor/Private/AssetSourceControlContextMenu.cpp
--- a/Engine/Plugins/Editor/AssetManagerEditor/Source/AssetManagerEditor/Private/AssetSourceControlContextMenu.cpp
+++ b/Engine/Plugins/Editor/AssetManagerEditor/Source/AssetManagerEditor/Private/AssetSourceControlContextMenu.cpp
@@ -12,6 +12,7 @@
 #include "ContentBrowserDataMenuContexts.h"
 #include "ContentBrowserMenuContexts.h"
 #include "FileHelpers.h"
+#include "HAL/IConsoleManager.h"
 #include "HAL/PlatformFile.h"
 #include "HAL/PlatformFileManager.h"
 #include "IAssetTools.h"
@@ -155,6 +156,8 @@ private:
 	bool CanExecuteSCCDiffAgainstDepot() const;
 	bool CanExecuteDiffSelected() const;
 
+	bool AllowExecuteSCCRevertUnsaved() const;
+
 	/** Helper function to gather the package names of all selected assets */
 	void GetSelectedPackageNames(TArray<FString>& OutPackageNames) const;
 
@@ -963,6 +966,18 @@ bool FAssetSourceControlContextMenuState::CanExecuteDiffSelected() const
 	return bCanDiffSelected;
 }
 
+bool FAssetSourceControlContextMenuState::AllowExecuteSCCRevertUnsaved() const
+{
+	if (IConsoleVariable* CVar = IConsoleManager::Get().FindConsoleVariable(TEXT("SourceControl.RevertUnsaved.Enable")))
+	{
+		return CVar->GetBool();
+	}
+	else
+	{
+		return false;
+	}
+}
+
 bool FAssetSourceControlContextMenuState::IsActionEnabled(bool bEnabled) const
 {
 	return bEnabled && !IsStillScanning(bEnabled);
@@ -1048,6 +1063,25 @@ void FAssetSourceControlContextMenuState::TryCacheCanExecuteVars(const TArray<FS
 	const bool bUsesCheckout = SourceControlProvider.UsesCheckout();
 	const bool bUsesFileRevisions = SourceControlProvider.UsesFileRevisions();
 
+	// If a package is dirty, allow a revert of the in-memory changes that have not yet been saved to disk.
+	if (AllowExecuteSCCRevertUnsaved())
+	{
+		for (const FString& SelectedPath : InSelectedPaths)
+		{
+			FString PackageName;
+			if (FPackageName::TryConvertFilenameToLongPackageName(SelectedPath, PackageName))
+			{
+				if (UPackage* Package = FindPackage(NULL, *PackageName))
+				{
+					if (Package->IsDirty())
+					{
+						bCanExecuteSCCRevert = true;
+					}
+				}
+			}
+		}
+	}
+
 	// Check the SCC state for each package in the selected paths
 	TArray<FSourceControlStateRef> SelectedStates;
 	if (SourceControlProvider.GetState(InSelectedPaths, SelectedStates, EStateCacheUsage::Use) == ECommandResult::Succeeded)
diff --git a/Engine/Source/Developer/SourceControl/Private/SourceControlCVars.cpp b/Engine/Source/Developer/SourceControl/Private/SourceControlCVars.cpp
--- a/Engine/Source/Developer/SourceControl/Private/SourceControlCVars.cpp
+++ b/Engine/Source/Developer/SourceControl/Private/SourceControlCVars.cpp
@@ -16,6 +16,12 @@ namespace SourceControlCVars
 		TEXT("Allows a SourceControl 'Revert' operation to be triggered from the SubmitWidget."),
 		ECVF_Default);
 
+	TAutoConsoleVariable<bool> CVarSourceControlEnableRevertUnsaved(
+		TEXT("SourceControl.RevertUnsaved.Enable"),
+		false,
+		TEXT("Allows a SourceControl 'Revert' operation to be triggered on an unsaved asset."),
+		ECVF_Default);
+
 	TAutoConsoleVariable<bool> CVarSourceControlEnableLoginDialogModal(
 		TEXT("SourceControl.LoginDialog.ForceModal"),
 		false,
diff --git a/Engine/Source/Developer/SourceControl/Private/SourceControlCVars.h b/Engine/Source/Developer/SourceControl/Private/SourceControlCVars.h
--- a/Engine/Source/Developer/SourceControl/Private/SourceControlCVars.h
+++ b/Engine/Source/Developer/SourceControl/Private/SourceControlCVars.h
@@ -9,5 +9,6 @@ namespace SourceControlCVars
 {
 	extern TAutoConsoleVariable<bool> CVarSourceControlEnableRevertFromSceneOutliner;
 	extern TAutoConsoleVariable<bool> CVarSourceControlEnableRevertFromSubmitWidget;
+	extern TAutoConsoleVariable<bool> CVarSourceControlEnableRevertUnsaved;
 	extern TAutoConsoleVariable<bool> CVarSourceControlEnableLoginDialogModal;
 }
\ No newline at end of file
diff --git a/Engine/Source/Editor/SceneOutliner/Private/SceneOutlinerSCCHandler.cpp b/Engine/Source/Editor/SceneOutliner/Private/SceneOutlinerSCCHandler.cpp
--- a/Engine/Source/Editor/SceneOutliner/Private/SceneOutlinerSCCHandler.cpp
+++ b/Engine/Source/Editor/SceneOutliner/Private/SceneOutlinerSCCHandler.cpp
@@ -76,15 +76,21 @@ bool FSceneOutlinerSCCHandler::AddSourceControlMenuOptions(UToolMenu* Menu, TArr
 	return false;
 }
 
-
-bool FSceneOutlinerSCCHandler::CanExecuteSourceControlActions() const
+bool FSceneOutlinerSCCHandler::AllowExecuteSourceControlRevert() const
 {
-	return SelectedItems.Num() > 0;
+	if (IConsoleVariable* CVar = IConsoleManager::Get().FindConsoleVariable(TEXT("SourceControl.Revert.EnableFromSceneOutliner")))
+	{
+		return CVar->GetBool();
+	}
+	else
+	{
+		return false;
+	}
 }
 
-bool FSceneOutlinerSCCHandler::CanExecuteSourceControlRevert() const
+bool FSceneOutlinerSCCHandler::AllowExecuteSourceControlRevertUnsaved() const
 {
-	if (IConsoleVariable* CVar = IConsoleManager::Get().FindConsoleVariable(TEXT("SourceControl.Revert.EnableFromSceneOutliner")))
+	if (IConsoleVariable* CVar = IConsoleManager::Get().FindConsoleVariable(TEXT("SourceControl.RevertUnsaved.Enable")))
 	{
 		return CVar->GetBool();
 	}
@@ -94,6 +100,11 @@ bool FSceneOutlinerSCCHandler::CanExecuteSourceControlRevert() const
 	}
 }
 
+bool FSceneOutlinerSCCHandler::CanExecuteSourceControlActions() const
+{
+	return SelectedItems.Num() > 0;
+}
+
 void FSceneOutlinerSCCHandler::CacheCanExecuteVars()
 {
 	bCanExecuteSCC = true;
@@ -104,6 +115,9 @@ void FSceneOutlinerSCCHandler::CacheCanExecuteVars()
 	bUsesFileRevisions = false;
 	bUsesChangelists = false;
 
+	const bool bAllowRevert = AllowExecuteSourceControlRevert();
+	const bool bAllowRevertUnsaved = AllowExecuteSourceControlRevertUnsaved();
+
 	if ( ISourceControlModule::Get().IsEnabled() || FUncontrolledChangelistsModule::Get().IsEnabled())
 	{
 		for (FSceneOutlinerTreeItemPtr SelectedItem : SelectedItems)
@@ -145,19 +159,21 @@ void FSceneOutlinerSCCHandler::CacheCanExecuteVars()
 					bCanExecuteSCCCheckIn = true;
 				}
 
-				bool bAllowRevert = CanExecuteSourceControlRevert();
-				if ( SourceControlState->CanRevert() && bAllowRevert )
-				{
-					bCanExecuteSCCRevert = true;
-				}
-				else
+				if ( bAllowRevert )
 				{
-					// If the package is dirty, allow a revert of the in-memory changes that have not yet been saved to disk.
-					if (UPackage* Package = SourceControl->GetPackage())
+					if ( SourceControlState->CanRevert() )
+					{
+						bCanExecuteSCCRevert = true;
+					}
+					else if ( bAllowRevertUnsaved )
 					{
-						if (Package->IsDirty())
+						// If the package is dirty, allow a revert of the in-memory changes that have not yet been saved to disk.
+						if (UPackage* Package = SourceControl->GetPackage())
 						{
-							bCanExecuteSCCRevert = true;
+							if (Package->IsDirty())
+							{
+								bCanExecuteSCCRevert = true;
+							}
 						}
 					}
 				}
diff --git a/Engine/Source/Editor/SceneOutliner/Public/SceneOutlinerSCCHandler.h b/Engine/Source/Editor/SceneOutliner/Public/SceneOutlinerSCCHandler.h
--- a/Engine/Source/Editor/SceneOutliner/Public/SceneOutlinerSCCHandler.h
+++ b/Engine/Source/Editor/SceneOutliner/Public/SceneOutlinerSCCHandler.h
@@ -17,8 +17,9 @@ public:
 
 private:
 
+	bool AllowExecuteSourceControlRevert() const;
+	bool AllowExecuteSourceControlRevertUnsaved() const;
 	bool CanExecuteSourceControlActions() const;
-	bool CanExecuteSourceControlRevert() const;
 	void CacheCanExecuteVars();
 	bool CanExecuteSCC() const;
 	bool CanExecuteSCCCheckOut() const;
diff --git a/Engine/Source/Editor/SourceControlWindows/Private/SSourceControlRevert.cpp b/Engine/Source/Editor/SourceControlWindows/Private/SSourceControlRevert.cpp
--- a/Engine/Source/Editor/SourceControlWindows/Private/SSourceControlRevert.cpp
+++ b/Engine/Source/Editor/SourceControlWindows/Private/SSourceControlRevert.cpp
@@ -33,6 +33,7 @@
 #include "Framework/Notifications/NotificationManager.h"
 #include "Widgets/Notifications/SNotificationList.h"
 #include "UObject/Linker.h"
+#include "FileHelpers.h"
 
 #define LOCTEXT_NAMESPACE "SSourceControlRevert"
 
@@ -83,6 +84,19 @@ struct FRevertCheckBoxListViewItem
 	FString Text;
 };
 
+/** Returns whether revert unsaved is enabled */
+static bool IsRevertUnsavedEnabled()
+{
+	if (IConsoleVariable* CVar = IConsoleManager::Get().FindConsoleVariable(TEXT("SourceControl.RevertUnsaved.Enable")))
+	{
+		return CVar->GetBool();
+	}
+	else
+	{
+		return false;
+	}
+}
+
 /**
  * Source control panel for reverting files. Allows the user to select which files should be reverted, as well as
  * provides the option to only allow unmodified files to be reverted.
@@ -360,6 +374,7 @@ private:
 
 		ModifiedPackages.Empty();
 
+		const bool bRevertUnsaved = IsRevertUnsavedEnabled();
 		for( const auto& ControlState : SourceControlStates )
 		{
 			FString PackageName;
@@ -369,6 +384,15 @@ private:
 				if (CurItem->Text == PackageName)
 				{
 					CurItem->IsModified = ControlState->IsModified();
+
+					if (bRevertUnsaved)
+					{
+						if (UPackage* Package = FindPackage(NULL, *PackageName))
+						{
+							// If the package contains unsaved changes, it's considered modified as well.
+							CurItem->IsModified |= Package->IsDirty();
+						}
+					}
 				}
 			}
 		}
@@ -412,10 +436,20 @@ bool FSourceControlWindows::PromptForRevert( const TArray<FString>& InPackageNam
 		{
 			InitialPackagesToRevert.Add( *PackageIter );
 		}
+		else if ( IsRevertUnsavedEnabled() )
+		{
+			if (UPackage* Package = FindPackage(NULL, **PackageIter))
+			{
+				if (Package->IsDirty())
+				{
+					InitialPackagesToRevert.Add(*PackageIter);
+				}
+			}
+		}
 	}
 
 	// If any of the packages can be reverted, provide the revert prompt
-	if (InitialPackagesToRevert.Num() > 0 )
+	if (InitialPackagesToRevert.Num() > 0)
 	{
 		TSharedRef<SWindow> NewWindow = SNew(SWindow)
 			.Title( NSLOCTEXT("SourceControl.RevertWindow", "Title", "Revert Files") )
@@ -438,6 +472,28 @@ bool FSourceControlWindows::PromptForRevert( const TArray<FString>& InPackageNam
 			TArray<FString> FinalPackagesToRevert;
 			SourceControlWidget->GetPackagesToRevert(FinalPackagesToRevert);
 			
+			if ( IsRevertUnsavedEnabled() )
+			{
+				// Unsaved changes need to be saved to disk so SourceControl realizes that there's something to revert.
+
+				TArray<UPackage*> FinalPackagesToSave;
+				for (const FString& PackageName : FinalPackagesToRevert)
+				{
+					if (UPackage* Package = FindPackage(NULL, *PackageName))
+					{
+						if (Package->IsDirty())
+						{
+							FinalPackagesToSave.Add(Package);
+						}
+					}
+				}
+
+				if (FinalPackagesToSave.Num() > 0)
+				{
+					UEditorLoadingAndSavingUtils::SavePackages(FinalPackagesToSave, /*bOnlyDirty=*/false);
+				}
+			}
+
 			if (FinalPackagesToRevert.Num() > 0)
 			{
 				SourceControlHelpers::RevertAndReloadPackages(FinalPackagesToRevert, /*bRevertAll=*/false, /*bReloadWorld=*/bInReloadWorld);
