SourceControl: Revert individual added non-package filse in check-in window sometimes does not work.

diff --git a/Engine/Source/Editor/SourceControlWindows/Private/SSourceControlSubmit.cpp b/Engine/Source/Editor/SourceControlWindows/Private/SSourceControlSubmit.cpp
--- a/Engine/Source/Editor/SourceControlWindows/Private/SSourceControlSubmit.cpp
+++ b/Engine/Source/Editor/SourceControlWindows/Private/SSourceControlSubmit.cpp
@@ -574,8 +574,7 @@ void SSourceControlSubmitWidget::OnRevert()
 	TArray<FString> FilesToRevert;
 	for (const auto& SelectedItem : SelectedItems)
 	{
-		FString PackageName;
-		if (FPackageName::TryConvertFilenameToLongPackageName(SelectedItem->GetFileName().ToString(), PackageName))
+		if (FPackageName::IsPackageFilename(SelectedItem->GetFileName().ToString()))
 		{
 			PackagesToRevert.Add(SelectedItem->GetFileName().ToString());
 		}
