View Changes: Prevent possible crashes when reverting

diff --git a/Engine/Source/Editor/SourceControlWindows/Private/SSourceControlChangelistRows.cpp b/Engine/Source/Editor/SourceControlWindows/Private/SSourceControlChangelistRows.cpp
--- a/Engine/Source/Editor/SourceControlWindows/Private/SSourceControlChangelistRows.cpp
+++ b/Engine/Source/Editor/SourceControlWindows/Private/SSourceControlChangelistRows.cpp
@@ -561,13 +561,16 @@ TSharedRef<SWidget> SOfflineFileTableRow::GenerateWidgetForColumn(const FName& C
 						}
 					}
 					// Validate we have a saved map
-					UPackage* LevelPackage = FindPackage(nullptr, *PackageName)->GetOutermost();
-					if (LevelPackage == GetTransientPackage()
-						|| LevelPackage->HasAnyFlags(RF_Transient)
-						|| !FPackageName::IsValidLongPackageName(LevelPackage->GetName()))
+					if (UPackage* Package = FindPackage(nullptr, *PackageName))
 					{
-						FMessageDialog::Open(EAppMsgType::Ok, LOCTEXT("DiscardUnsavedChangesSaveMap", "You need to save the level before discarding unsaved changes."));
-						return FReply::Handled();
+						UPackage* LevelPackage = Package->GetOutermost();
+						if (LevelPackage == GetTransientPackage()
+							|| LevelPackage->HasAnyFlags(RF_Transient)
+							|| !FPackageName::IsValidLongPackageName(LevelPackage->GetName()))
+						{
+							FMessageDialog::Open(EAppMsgType::Ok, LOCTEXT("DiscardUnsavedChangesSaveMap", "You need to save the level before discarding unsaved changes."));
+							return FReply::Handled();
+						}
 					}
 					
 					DiscardSwitcher->SetActiveWidgetIndex(1);
