View Changes: Refresh active view when sorting columns

Previously only refreshed the changelist view

diff --git a/Engine/Source/Editor/SourceControlWindows/Private/SSourceControlChangelists.cpp b/Engine/Source/Editor/SourceControlWindows/Private/SSourceControlChangelists.cpp
--- a/Engine/Source/Editor/SourceControlWindows/Private/SSourceControlChangelists.cpp
+++ b/Engine/Source/Editor/SourceControlWindows/Private/SSourceControlChangelists.cpp
@@ -3515,7 +3515,7 @@ void SSourceControlChangelistsWidget::OnColumnSortModeChanged(const EColumnSortP
 	}
 
 	SortFileView();
-	FileListView->RequestListRefresh();
+	GetActiveFileListView().RequestListRefresh();
 }
 
 bool SSourceControlChangelistsWidget::IsFileViewSortedByFileStatusIcon() const
