View Changes: Saving to changelists can now be canceled

diff --git a/Engine/Source/Editor/SourceControlWindows/Private/SSourceControlChangelists.cpp b/Engine/Source/Editor/SourceControlWindows/Private/SSourceControlChangelists.cpp
--- a/Engine/Source/Editor/SourceControlWindows/Private/SSourceControlChangelists.cpp
+++ b/Engine/Source/Editor/SourceControlWindows/Private/SSourceControlChangelists.cpp
@@ -2592,7 +2592,7 @@ bool SSourceControlChangelistsWidget::CanDeleteUncontrolledChangelist()
 	return (UncontrolledChangelistState != nullptr) && !UncontrolledChangelistState->Changelist.IsDefault() && !UncontrolledChangelistState->ContainsFiles();
 }
 
-void SSourceControlChangelistsWidget::OnMoveFiles()
+TValueOrError<void, void> SSourceControlChangelistsWidget::TryMoveFiles()
 {
 	TArray<FString> SelectedControlledFiles;
 	TArray<FString> SelectedUncontrolledFiles;
@@ -2601,7 +2601,7 @@ void SSourceControlChangelistsWidget::OnMoveFiles()
 
 	if (SelectedControlledFiles.IsEmpty() && SelectedUncontrolledFiles.IsEmpty())
 	{
-		return;
+		return MakeError();
 	}
 
 	const bool bAddNewChangelistEntry = true;
@@ -2662,11 +2662,13 @@ void SSourceControlChangelistsWidget::OnMoveFiles()
 
 	if (!bOk)
 	{
-		return;
+		return MakeError();
 	}
 
 	ISourceControlProvider& SourceControlProvider = ISourceControlModule::Get().GetProvider();
 
+	bool bFailed = false;
+	
 	// Move files to a new changelist
 	if (bAddNewChangelistEntry && PickedItem == 0)
 	{
@@ -2674,7 +2676,7 @@ void SSourceControlChangelistsWidget::OnMoveFiles()
 		TSharedRef<FNewChangelist> NewChangelistOperation = ISourceControlOperation::Create<FNewChangelist>();
 		NewChangelistOperation->SetDescription(ChangelistDescription);
 		Execute(LOCTEXT("Moving_Files_New_Changelist", "Moving file(s) to a new changelist..."), NewChangelistOperation, SelectedControlledFiles, EConcurrency::Synchronous, FSourceControlOperationComplete::CreateLambda(
-			[this, SelectedUncontrolledFiles](const TSharedRef<ISourceControlOperation>& Operation, ECommandResult::Type InResult)
+			[this, &SelectedUncontrolledFiles, &bFailed](const TSharedRef<ISourceControlOperation>& Operation, ECommandResult::Type InResult)
 			{
 				if (InResult == ECommandResult::Succeeded)
 				{
@@ -2689,6 +2691,7 @@ void SSourceControlChangelistsWidget::OnMoveFiles()
 				if (InResult == ECommandResult::Failed)
 				{
 					SSourceControlCommon::DisplaySourceControlOperationNotification(LOCTEXT("Move_Files_New_Changelist_Failed", "Failed to move the file to the new changelist."), SNotificationItem::CS_Fail);
+					bFailed = true;
 				}
 			}));
 	}
@@ -2707,6 +2710,7 @@ void SSourceControlChangelistsWidget::OnMoveFiles()
 			if (!NewUncontrolledChangelist.IsSet())
 			{
 				SSourceControlCommon::DisplaySourceControlOperationNotification(LOCTEXT("Move_Files_New_Uncontrolled_Changelist_Failed", "Failed to create a new uncontrolled changelist."), SNotificationItem::CS_Fail);
+				bFailed = true;
 			}
 			else if (!SelectedControlledFileStates.IsEmpty() || !SelectedUnControlledFileStates.IsEmpty())
 			{
@@ -2743,7 +2747,7 @@ void SSourceControlChangelistsWidget::OnMoveFiles()
 			if (!SelectedControlledFiles.IsEmpty())
 			{
 				Execute(LOCTEXT("Moving_File_Between_Changelists", "Moving file(s) to the selected changelist..."), ISourceControlOperation::Create<FMoveToChangelist>(), Changelist, SelectedControlledFiles, EConcurrency::Synchronous, FSourceControlOperationComplete::CreateLambda(
-					[](const TSharedRef<ISourceControlOperation>& Operation, ECommandResult::Type InResult)
+					[&bFailed](const TSharedRef<ISourceControlOperation>& Operation, ECommandResult::Type InResult)
 					{
 						if (InResult == ECommandResult::Succeeded)
 						{
@@ -2752,6 +2756,7 @@ void SSourceControlChangelistsWidget::OnMoveFiles()
 						else if (InResult == ECommandResult::Failed)
 						{
 							SSourceControlCommon::DisplaySourceControlOperationNotification(LOCTEXT("Move_Files_Between_Changelist_Failed", "Failed to move the file(s) to the selected changelist."), SNotificationItem::CS_Fail);
+							bFailed = true;
 						}
 					}));
 			}
@@ -2781,6 +2786,11 @@ void SSourceControlChangelistsWidget::OnMoveFiles()
 			}
 		}
 	}
+	if (bFailed)
+	{
+		return MakeError();
+	}
+	return MakeValue();
 }
 
 void SSourceControlChangelistsWidget::OnShowHistory()
@@ -2900,7 +2910,10 @@ TSharedPtr<SWidget> SSourceControlChangelistsWidget::OnOpenContextMenu()
 				}
 				
 				// Handles almost everything else
-				OnMoveFiles();
+				if (TryMoveFiles().HasError())
+				{
+					return;
+				}
 
 				// Actaully save the assets 
 				TArray<UPackage*> Packages;
@@ -3056,7 +3069,7 @@ TSharedPtr<SWidget> SSourceControlChangelistsWidget::OnOpenContextMenu()
 			"MoveFiles", LOCTEXT("SourceControl_MoveFiles", "Move Files To..."),
 			LOCTEXT("SourceControl_MoveFiles_Tooltip", "Move Files To A Different Changelist..."),
 			FSlateIcon(),
-			FUIAction(FExecuteAction::CreateSP(this, &SSourceControlChangelistsWidget::OnMoveFiles)));
+			FUIAction(FExecuteAction::CreateLambda([this] { TryMoveFiles(); })));
 
 		Section.AddMenuEntry(
 			"ShowHistory",
diff --git a/Engine/Source/Editor/SourceControlWindows/Private/SSourceControlChangelists.h b/Engine/Source/Editor/SourceControlWindows/Private/SSourceControlChangelists.h
--- a/Engine/Source/Editor/SourceControlWindows/Private/SSourceControlChangelists.h
+++ b/Engine/Source/Editor/SourceControlWindows/Private/SSourceControlChangelists.h
@@ -164,7 +164,7 @@ private:
 	void OnDeleteShelvedFiles();
 
 	/** Files operations */
-	void OnMoveFiles();
+	TValueOrError<void, void> TryMoveFiles();
 	void OnShowHistory();
 	void OnDiffAgainstDepot();
 	bool CanDiffAgainstDepot();
