Added ability to revert conflicted files in the Submit Files window.

A side effect of showing conflicted items greyed out in the Submit Files window is that the Revert operation is also disabled as greyed-out items are not selectable. This is not desirable as the user may well want to get out of a conflict situation by reverting the conflicted file and the submit dialog is a convenient place to do that.

diff --git a/Engine/Source/Editor/SourceControlWindows/Private/SSourceControlSubmit.cpp b/Engine/Source/Editor/SourceControlWindows/Private/SSourceControlSubmit.cpp
--- a/Engine/Source/Editor/SourceControlWindows/Private/SSourceControlSubmit.cpp
+++ b/Engine/Source/Editor/SourceControlWindows/Private/SSourceControlSubmit.cpp
@@ -807,8 +807,7 @@ TSharedRef<ITableRow> SSourceControlSubmitWidget::OnGenerateRowForList(TSharedPt
 	TSharedRef<ITableRow> Row =
 	SNew(SSourceControlSubmitListRow, OwnerTable)
 		.SourceControlSubmitWidget(SharedThis(this))
-		.Item(SubmitItem)
-		.IsEnabled(SubmitItem->IsEnabled());
+		.Item(SubmitItem);
 
 	return Row;
 }
