View Changes: Only reregister components once when saving

diff --git a/Engine/Source/Editor/SourceControlWindows/Private/SSourceControlChangelists.cpp b/Engine/Source/Editor/SourceControlWindows/Private/SSourceControlChangelists.cpp
--- a/Engine/Source/Editor/SourceControlWindows/Private/SSourceControlChangelists.cpp
+++ b/Engine/Source/Editor/SourceControlWindows/Private/SSourceControlChangelists.cpp
@@ -32,6 +32,7 @@
 #include "SourceControlMenuContext.h"
 #include "SourceControlSettings.h"
 #include "AssetToolsModule.h"
+#include "ComponentReregisterContext.h"
 #include "FileHelpers.h"
 #include "Misc/CString.h"
 #include "Misc/MessageDialog.h"
@@ -2879,10 +2880,15 @@ TSharedPtr<SWidget> SSourceControlChangelistsWidget::OnOpenContextMenu()
 				TArray<FString> SelectedFiles;
 				GetSelectedFiles(SelectedFiles, SelectedFiles);
 				TArray<UPackage*> Packages;
-				for (const FString& Filename : SelectedFiles)
+				Packages.Reserve(SelectedFiles.Num());
 				{
-					FString PackageName = UPackageTools::FilenameToPackageName(Filename);
-					Packages.Add(UPackageTools::LoadPackage(PackageName));
+					// Reregister here because otherwise LoadPackage will do it once per call and it can be very slow
+					FGlobalComponentReregisterContext ReregisterContext;
+					for (const FString& Filename : SelectedFiles)
+					{
+						FString PackageName = UPackageTools::FilenameToPackageName(Filename);
+						Packages.Add(UPackageTools::LoadPackage(PackageName));
+					}
 				}
 				const bool bOnlyDirty = false;
 				UEditorLoadingAndSavingUtils::SavePackages(Packages, bOnlyDirty);
@@ -2915,11 +2921,17 @@ TSharedPtr<SWidget> SSourceControlChangelistsWidget::OnOpenContextMenu()
 					return;
 				}
 
-				// Actaully save the assets 
+				// Actually save the assets
 				TArray<UPackage*> Packages;
-				for (const FString& PackageName : PackageNames)
+				Packages.Reserve(PackageNames.Num());
 				{
-					Packages.Add(UPackageTools::LoadPackage(PackageName));
+					// Reregister here because otherwise LoadPackage will do it once per call and it can be very slow
+					FGlobalComponentReregisterContext ReregisterContext;
+
+					for (const FString& PackageName : PackageNames)
+					{
+						Packages.Add(UPackageTools::LoadPackage(PackageName));
+					}
 				}
 				const bool bOnlyDirty = false;
 				UEditorLoadingAndSavingUtils::SavePackages(Packages, bOnlyDirty);
