Fix crash when diffing virtual revision

diff --git a/Engine/Source/Editor/SourceControlWindows/Private/SSourceControlHistory.cpp b/Engine/Source/Editor/SourceControlWindows/Private/SSourceControlHistory.cpp
--- a/Engine/Source/Editor/SourceControlWindows/Private/SSourceControlHistory.cpp
+++ b/Engine/Source/Editor/SourceControlWindows/Private/SSourceControlHistory.cpp
@@ -1179,10 +1179,13 @@ private:
 				TSharedPtr<FHistoryTreeItem> Parent = Item->Parent.Pin();
 
 				USourceControlHistoryWidgetContext::SelectedItem ItemInfo;
-				ItemInfo.FileName = Parent->FileListItem->FileName;
-				ItemInfo.Revision = Item->RevisionListItem->Revision;
+				if (Parent->FileListItem)
+				{
+					ItemInfo.FileName = Parent->FileListItem->FileName;
+					ItemInfo.Revision = Item->RevisionListItem->Revision;
 
-				SourceControlHistoryWidgetContext->GetSelectedItems().Add(ItemInfo);
+					SourceControlHistoryWidgetContext->GetSelectedItems().Add(ItemInfo);
+				}
 			}
 		}
 	
