Made Submit Content button invisible when URC is the selected source control provider.

diff --git a/Engine/Source/Editor/SourceControlWindows/Private/SourceControlWindows.cpp b/Engine/Source/Editor/SourceControlWindows/Private/SourceControlWindows.cpp
--- a/Engine/Source/Editor/SourceControlWindows/Private/SourceControlWindows.cpp
+++ b/Engine/Source/Editor/SourceControlWindows/Private/SourceControlWindows.cpp
@@ -129,7 +129,7 @@ bool FSourceControlWindows::CanChoosePackagesToCheckIn()
 
 bool FSourceControlWindows::ShouldChoosePackagesToCheckBeVisible()
 {
-	return GetDefault<USourceControlSettings>()->bEnableSubmitContentMenuAction;
+	return GetDefault<USourceControlSettings>()->bEnableSubmitContentMenuAction && !ISourceControlModule::Get().GetProvider().UsesSnapshots();
 }
 
 
