New feature:
Unsaved asset drag-and-drop handler: When unsaved assets are selected and then dropped onto a changelist entry, they are checked out, saved, and then moved to the target changelist.

diff --git a/Engine/Source/Developer/UncontrolledChangelists/Private/UncontrolledChangelistsModule.cpp b/Engine/Source/Developer/UncontrolledChangelists/Private/UncontrolledChangelistsModule.cpp
--- a/Engine/Source/Developer/UncontrolledChangelists/Private/UncontrolledChangelistsModule.cpp
+++ b/Engine/Source/Developer/UncontrolledChangelists/Private/UncontrolledChangelistsModule.cpp
@@ -422,7 +422,7 @@ void FUncontrolledChangelistsModule::MoveFilesToUncontrolledChangelist(const TAr
 	Algo::Transform(InControlledFileStates, Filenames, [](const FSourceControlStateRef& State) { return State->GetFilename(); });
 
 	ISourceControlProvider& SourceControlProvider = ISourceControlModule::Get().GetProvider();
-	auto RevertOperation = ISourceControlOperation::Create<FRevert>();
+	TSharedRef<FRevert, ESPMode::ThreadSafe> RevertOperation = ISourceControlOperation::Create<FRevert>();
 
 	// Revert controlled files
 	RevertOperation->SetSoftRevert(true);
@@ -446,6 +446,40 @@ void FUncontrolledChangelistsModule::MoveFilesToUncontrolledChangelist(const TAr
 	}
 }
 
+void FUncontrolledChangelistsModule::MoveFilesToUncontrolledChangelist(const TArray<FString>& InControlledFiles, const FUncontrolledChangelist& InUncontrolledChangelist)
+{
+	bool bHasStateChanged = false;
+
+	if (!IsEnabled())
+	{
+		return;
+	}
+
+	FUncontrolledChangelistsStateCache::ValueType* ChangelistState = UncontrolledChangelistsStateCache.Find(InUncontrolledChangelist);
+
+	if (ChangelistState == nullptr)
+	{
+		return;
+	}
+
+	const TArray<FString>& Filenames = InControlledFiles;
+
+	ISourceControlProvider& SourceControlProvider = ISourceControlModule::Get().GetProvider();
+	TSharedRef<FRevert, ESPMode::ThreadSafe> RevertOperation = ISourceControlOperation::Create<FRevert>();
+
+	// Revert controlled files
+	RevertOperation->SetSoftRevert(true);
+	SourceControlProvider.Execute(RevertOperation, Filenames);
+
+	// Add all files to their UncontrolledChangelist
+	bHasStateChanged = (*ChangelistState)->AddFiles(Filenames, FUncontrolledChangelistState::ECheckFlags::None);
+
+	if (bHasStateChanged)
+	{
+		OnStateChanged();
+	}
+}
+
 void FUncontrolledChangelistsModule::MoveFilesToControlledChangelist(const TArray<FSourceControlStateRef>& InUncontrolledFileStates, const FSourceControlChangelistPtr& InChangelist, TFunctionRef<bool(const TArray<FSourceControlStateRef>&)> InOpenConflictDialog)
 {
 	if (!IsEnabled())
diff --git a/Engine/Source/Developer/UncontrolledChangelists/Public/UncontrolledChangelistsModule.h b/Engine/Source/Developer/UncontrolledChangelists/Public/UncontrolledChangelistsModule.h
--- a/Engine/Source/Developer/UncontrolledChangelists/Public/UncontrolledChangelistsModule.h
+++ b/Engine/Source/Developer/UncontrolledChangelists/Public/UncontrolledChangelistsModule.h
@@ -142,10 +142,17 @@ public:
 	 * Moves files to an Uncontrolled Changelist.
 	 * @param 	InControlledFileStates 		The Controlled files to move.
 	 * @param 	InUncontrolledFileStates 	The Uncontrolled files to move.
-	 * @param 	InChangelist 				The Uncontrolled Changelist where to move the files.
+	 * @param 	InUncontrolledChangelist 	The Uncontrolled Changelist where to move the files.
 	 */
 	void MoveFilesToUncontrolledChangelist(const TArray<FSourceControlStateRef>& InControlledFileStates, const TArray<FSourceControlStateRef>& InUncontrolledFileStates, const FUncontrolledChangelist& InUncontrolledChangelist);
 
+	/**
+	* Moves files to an Uncontrolled Changelist.
+	* @param 	InControlledFileStates 		The Controlled files to move.
+	* @param 	InUncontrolledChangelist 	The Uncontrolled Changelist where to move the files.
+	*/
+	void MoveFilesToUncontrolledChangelist(const TArray<FString>& InControlledFiles, const FUncontrolledChangelist& InUncontrolledChangelist);
+
 	/**
 	 * Moves files to a Controlled Changelist.
 	 * @param 	InUncontrolledFileStates 	The files to move.
diff --git a/Engine/Source/Editor/SourceControlWindows/Private/SSourceControlChangelistRows.cpp b/Engine/Source/Editor/SourceControlWindows/Private/SSourceControlChangelistRows.cpp
--- a/Engine/Source/Editor/SourceControlWindows/Private/SSourceControlChangelistRows.cpp
+++ b/Engine/Source/Editor/SourceControlWindows/Private/SSourceControlChangelistRows.cpp
@@ -2,8 +2,11 @@
 
 #include "SSourceControlChangelistRows.h"
 
+#include "ComponentReregisterContext.h"
+#include "FileHelpers.h"
 #include "ISourceControlModule.h"
 #include "PackageTools.h"
+#include "PackageSourceControlHelper.h"
 #include "SourceControlHelpers.h"
 #include "UncontrolledChangelistsModule.h"
 #include "SourceControlOperations.h"
@@ -23,6 +26,36 @@
 
 #define LOCTEXT_NAMESPACE "SourceControlChangelistRow"
 
+namespace
+{
+	bool CheckoutAndSavePackages(const TArray<FString>& Files)
+	{
+		FPackageSourceControlHelper PackageHelper;
+		TArray<FString> PackageNames;
+		PackageNames.Reserve(Files.Num());
+		for (const FString& Filename : Files)
+		{
+			PackageNames.Add(UPackageTools::FilenameToPackageName(Filename));
+		}
+
+		if (PackageHelper.Checkout(PackageNames))
+		{
+			TArray<UPackage*> Packages;
+			Packages.Reserve(PackageNames.Num());
+			{
+				FGlobalComponentReregisterContext ReregisterContext;
+				for (const FString& PackageName : PackageNames)
+				{
+					Packages.Add(UPackageTools::LoadPackage(PackageName));
+				}
+			}
+			constexpr bool bOnlyDirty = false;
+			return UEditorLoadingAndSavingUtils::SavePackages(Packages, bOnlyDirty);
+		}
+		return false;
+	}
+}
+
 FName SourceControlFileViewColumn::Icon::Id() { return TEXT("Icon"); }
 FText SourceControlFileViewColumn::Icon::GetDisplayText() {return LOCTEXT("Name_Icon", "Revision Control Status"); }
 FText SourceControlFileViewColumn::Icon::GetToolTipText() { return LOCTEXT("Icon_Column_Tooltip", "Displays the asset/file status"); }
@@ -183,6 +216,32 @@ FReply SChangelistTableRow::OnDrop(const FGeometry& InGeometry, const FDragDropE
 
 			OnPostDrop.ExecuteIfBound();
 		}
+		else if (!DropOperation->OfflineFiles.IsEmpty())
+		{
+			const TArray<FString>& Files = DropOperation->OfflineFiles;
+			SSourceControlCommon::ExecuteChangelistOperationWithSlowTaskWrapper(LOCTEXT("Dropping_Files_On_Changelist", "Moving file(s) to the selected changelist..."), [&]()
+				{
+					ISourceControlProvider& SourceControlProvider = ISourceControlModule::Get().GetProvider();
+
+					if (!CheckoutAndSavePackages(Files))
+					{
+						return;
+					}
+
+					SourceControlProvider.Execute(ISourceControlOperation::Create<FMoveToChangelist>(), DestChangelist, Files, EConcurrency::Synchronous, FSourceControlOperationComplete::CreateLambda(
+						[](const TSharedRef<ISourceControlOperation>& Operation, ECommandResult::Type InResult)
+						{
+							if (InResult == ECommandResult::Succeeded)
+							{
+								SSourceControlCommon::DisplaySourceControlOperationNotification(LOCTEXT("Drop_Files_On_Changelist_Succeeded", "File(s) successfully moved to the selected changelist."), SNotificationItem::CS_Success);
+							}
+							else if (InResult == ECommandResult::Failed)
+							{
+								SSourceControlCommon::DisplaySourceControlOperationNotification(LOCTEXT("Drop_Files_On_Changelist_Failed", "Failed to move the file(s) to the selected changelist."), SNotificationItem::CS_Fail);
+							}
+						}));
+				});
+		}
 	}
 
 	return FReply::Handled();
@@ -252,11 +311,29 @@ FReply SUncontrolledChangelistTableRow::OnDrop(const FGeometry& InGeometry, cons
 
 	if (Operation.IsValid())
 	{
-		SSourceControlCommon::ExecuteUncontrolledChangelistOperationWithSlowTaskWrapper(LOCTEXT("Drag_File_To_Uncontrolled_Changelist", "Moving file(s) to the selected uncontrolled changelists..."),
-			[this, &Operation]()
+		if (Operation->OfflineFiles.IsEmpty())
 		{
-			FUncontrolledChangelistsModule::Get().MoveFilesToUncontrolledChangelist(Operation->Files, Operation->UncontrolledFiles, TreeItem->UncontrolledChangelistState->Changelist);
-		});
+			SSourceControlCommon::ExecuteUncontrolledChangelistOperationWithSlowTaskWrapper(LOCTEXT("Drag_File_To_Uncontrolled_Changelist", "Moving file(s) to the selected uncontrolled changelists..."),
+				[this, &Operation]()
+				{
+					FUncontrolledChangelistsModule::Get().MoveFilesToUncontrolledChangelist(Operation->Files, Operation->UncontrolledFiles, TreeItem->UncontrolledChangelistState->Changelist);
+				});
+		}
+		// Drop unsaved assets (offline files)
+		else
+		{
+			const TArray<FString>& Files = Operation->OfflineFiles;
+			if (!CheckoutAndSavePackages(Files))
+			{
+				return FReply::Unhandled();
+			}
+
+			SSourceControlCommon::ExecuteUncontrolledChangelistOperationWithSlowTaskWrapper(LOCTEXT("Drag_File_To_Uncontrolled_Changelist", "Moving file(s) to the selected uncontrolled changelists..."),
+				[this, &Files]()
+				{
+					FUncontrolledChangelistsModule::Get().MoveFilesToUncontrolledChangelist(Files, TreeItem->UncontrolledChangelistState->Changelist);
+				});
+		}
 
 		OnPostDrop.ExecuteIfBound();
 	}
@@ -469,7 +546,9 @@ void SOfflineFileTableRow::Construct(const FArguments& InArgs, const TSharedRef<
 	TreeItem = static_cast<FOfflineFileTreeItem*>(InArgs._TreeItemToVisualize.Get());
 	HighlightText = InArgs._HighlightText;
 
-	FSuperRowType::FArguments Args = FSuperRowType::FArguments().ShowSelection(true);
+	FSuperRowType::FArguments Args = FSuperRowType::FArguments()
+		.OnDragDetected(InArgs._OnDragDetected)
+		.ShowSelection(true);
 	FSuperRowType::Construct(Args, InOwner);
 }
 
@@ -653,5 +732,17 @@ FText SOfflineFileTableRow::GetLastModifiedTimestamp() const
 	return TreeItem->GetLastModifiedTimestamp();
 }
 
+void SOfflineFileTableRow::OnDragEnter(FGeometry const& InGeometry, FDragDropEvent const& InDragDropEvent)
+{
+	TSharedPtr<FDragDropOperation> DragOperation = InDragDropEvent.GetOperation();
+	DragOperation->SetCursorOverride(EMouseCursor::SlashedCircle);
+}
+
+void SOfflineFileTableRow::OnDragLeave(FDragDropEvent const& InDragDropEvent)
+{
+	TSharedPtr<FDragDropOperation> DragOperation = InDragDropEvent.GetOperation();
+	DragOperation->SetCursorOverride(EMouseCursor::None);
+}
+
 
 #undef LOCTEXT_NAMESPACE
diff --git a/Engine/Source/Editor/SourceControlWindows/Private/SSourceControlChangelistRows.h b/Engine/Source/Editor/SourceControlWindows/Private/SSourceControlChangelistRows.h
--- a/Engine/Source/Editor/SourceControlWindows/Private/SSourceControlChangelistRows.h
+++ b/Engine/Source/Editor/SourceControlWindows/Private/SSourceControlChangelistRows.h
@@ -266,6 +266,7 @@ public:
 	}
 		SLATE_ARGUMENT(FChangelistTreeItemPtr, TreeItemToVisualize)
 		SLATE_ATTRIBUTE(FText, HighlightText)
+		SLATE_EVENT(FOnDragDetected, OnDragDetected)
 	SLATE_END_ARGS()
 
 public:
@@ -289,6 +290,12 @@ public:
 
 	static void PopulateSearchString(const FOfflineFileTreeItem& Item, TArray<FString>& OutStrings);
 
+protected:
+	//~ Begin STableRow Interface.
+	virtual void OnDragEnter(FGeometry const& InGeometry, FDragDropEvent const& InDragDropEvent) override;
+	virtual void OnDragLeave(FDragDropEvent const& InDragDropEvent) override;
+	//~ End STableRow Interface.
+
 private:
 	/** The info about the widget that we are visualizing. */
 	FOfflineFileTreeItem* TreeItem;
diff --git a/Engine/Source/Editor/SourceControlWindows/Private/SSourceControlChangelists.cpp b/Engine/Source/Editor/SourceControlWindows/Private/SSourceControlChangelists.cpp
--- a/Engine/Source/Editor/SourceControlWindows/Private/SSourceControlChangelists.cpp
+++ b/Engine/Source/Editor/SourceControlWindows/Private/SSourceControlChangelists.cpp
@@ -3348,7 +3348,8 @@ TSharedRef<ITableRow> SSourceControlChangelistsWidget::OnGenerateRow(TSharedPtr<
 		bUpdateMonitoredFileStatusList = true;
 		return SNew(SOfflineFileTableRow, OwnerTable)
 			.TreeItemToVisualize(InTreeItem)
-			.HighlightText_Lambda([this]() { return FileSearchBox->GetText(); });
+			.HighlightText_Lambda([this]() { return FileSearchBox->GetText(); })
+			.OnDragDetected(this, &SSourceControlChangelistsWidget::OnUnsavedAssetsDragged);
 
 	case IChangelistTreeItem::ShelvedChangelist:
 		return SNew(SShelvedFilesTableRow, OwnerTable)
@@ -3403,6 +3404,30 @@ FReply SSourceControlChangelistsWidget::OnFilesDragged(const FGeometry& InGeomet
 	return FReply::Unhandled();
 }
 
+FReply SSourceControlChangelistsWidget::OnUnsavedAssetsDragged(const FGeometry& InGeometry, const FPointerEvent& InMouseEvent)
+{
+	if (InMouseEvent.IsMouseButtonDown(EKeys::LeftMouseButton) && !UnsavedAssetsFileListView->GetSelectedItems().IsEmpty())
+	{
+		TSharedRef<FSCCFileDragDropOp> Operation = MakeShared<FSCCFileDragDropOp>();
+
+		for (const FChangelistTreeItemPtr& InTreeItem : UnsavedAssetsFileListView->GetSelectedItems())
+		{
+			if (InTreeItem->GetTreeItemType() == IChangelistTreeItem::OfflineFile)
+			{
+				TSharedRef<FOfflineFileTreeItem> FileTreeItem = StaticCastSharedRef<FOfflineFileTreeItem>(InTreeItem.ToSharedRef());
+				Operation->OfflineFiles.Emplace(FileTreeItem->GetFilename());
+			}
+		}
+
+		Operation->Construct();
+
+		// Initiates drag-drop operation.
+		return FReply::Handled().BeginDragDrop(Operation);
+	}
+
+	return FReply::Unhandled();
+}
+
 void SSourceControlChangelistsWidget::OnGetChangelistChildren(FChangelistTreeItemPtr InParent, TArray<FChangelistTreeItemPtr>& OutChildren)
 {
 	if (InParent->GetTreeItemType() == IChangelistTreeItem::Changelist)
diff --git a/Engine/Source/Editor/SourceControlWindows/Private/SSourceControlChangelists.h b/Engine/Source/Editor/SourceControlWindows/Private/SSourceControlChangelists.h
--- a/Engine/Source/Editor/SourceControlWindows/Private/SSourceControlChangelists.h
+++ b/Engine/Source/Editor/SourceControlWindows/Private/SSourceControlChangelists.h
@@ -88,6 +88,7 @@ private:
 	void OnUnsavedAssetChanged(const FString& Filepath);
 
 	FReply OnFilesDragged(const FGeometry& InGeometry, const FPointerEvent& InMouseEvent);
+	FReply OnUnsavedAssetsDragged(const FGeometry& InGeometry, const FPointerEvent& InMouseEvent);
 
 	void RequestChangelistsRefresh();
 	void RequestFileStatusRefresh(const IChangelistTreeItem& Changelist);
diff --git a/Engine/Source/Editor/SourceControlWindows/Private/SSourceControlCommon.cpp b/Engine/Source/Editor/SourceControlWindows/Private/SSourceControlCommon.cpp
--- a/Engine/Source/Editor/SourceControlWindows/Private/SSourceControlCommon.cpp
+++ b/Engine/Source/Editor/SourceControlWindows/Private/SSourceControlCommon.cpp
@@ -364,6 +364,33 @@ TSharedRef<SWidget> GetSCCFileWidget(FSourceControlStateRef InFileState, bool bI
 		];
 }
 
+TSharedRef<SWidget> GetSCCFileWidget()
+{
+	const FSlateBrush* IconBrush = FAppStyle::GetBrush("ContentBrowser.ColumnViewAssetIcon");
+
+	// Make icon overlays (eg, SCC and dirty status) a reasonable size in relation to the icon size (note: it is assumed this icon is square)
+	const float ICON_SCALING_FACTOR = 0.7f;
+	const float IconOverlaySize = IconBrush->ImageSize.X * ICON_SCALING_FACTOR;
+
+	return SNew(SOverlay)
+		// The actual icon
+		+ SOverlay::Slot()
+		[
+			SNew(SImage)
+			.Image(IconBrush)
+		.ColorAndOpacity(FSlateColor(FColor::White))
+		]
+	// Source control state
+	+ SOverlay::Slot()
+		.HAlign(HAlign_Left)
+		.VAlign(VAlign_Top)
+		[
+			SNew(SBox)
+			.WidthOverride(IconOverlaySize)
+		.HeightOverride(IconOverlaySize)
+		];
+}
+
 FText GetDefaultAssetName()
 {
 	return LOCTEXT("SourceControl_DefaultAssetName", "Unavailable");
diff --git a/Engine/Source/Editor/SourceControlWindows/Private/SSourceControlCommon.h b/Engine/Source/Editor/SourceControlWindows/Private/SSourceControlCommon.h
--- a/Engine/Source/Editor/SourceControlWindows/Private/SSourceControlCommon.h
+++ b/Engine/Source/Editor/SourceControlWindows/Private/SSourceControlCommon.h
@@ -377,6 +377,7 @@ private:
 namespace SSourceControlCommon
 {
 	TSharedRef<SWidget> GetSCCFileWidget(FSourceControlStateRef InFileState, bool bIsShelvedFile = false);
+	TSharedRef<SWidget> GetSCCFileWidget();
 	FText GetDefaultAssetName();
 	FText GetDefaultAssetType();
 	FText GetDefaultUnknownAssetType();
@@ -403,12 +404,19 @@ struct FSCCFileDragDropOp : public FDragDropOperation
 
 	virtual TSharedPtr<SWidget> GetDefaultDecorator() const override
 	{
+		// Offline files won't coexist with Files
+		if (!OfflineFiles.IsEmpty())
+		{
+			return SSourceControlCommon::GetSCCFileWidget();
+		}
+
 		FSourceControlStateRef FileState = Files.IsEmpty() ? UncontrolledFiles[0] : Files[0];
 		return SSourceControlCommon::GetSCCFileWidget(MoveTemp(FileState));
 	}
 
 	TArray<FSourceControlStateRef> Files;
 	TArray<FSourceControlStateRef> UncontrolledFiles;
+	TArray<FString> OfflineFiles;
 };
 
 
