New Delegate to override the submit functionality for the Source Control Changelist window and the Source Control Window.

This new delegate give users the capability to intercept and handle the submit instead of the regular flows.
Depending on the reply from the delegate the current flow will be:
- interrupted if the submit is handled by the delegate
- interrupted and set to fail state if the delegate returns an error
- continued if the delegate cannot handle the current source provider

diff --git a/Engine/Source/Editor/SourceControlWindows/Private/SSourceControlChangelists.cpp b/Engine/Source/Editor/SourceControlWindows/Private/SSourceControlChangelists.cpp
--- a/Engine/Source/Editor/SourceControlWindows/Private/SSourceControlChangelists.cpp
+++ b/Engine/Source/Editor/SourceControlWindows/Private/SSourceControlChangelists.cpp
@@ -2354,6 +2354,58 @@ void SSourceControlChangelistsWidget::OnSubmitChangelist()
 	// The description once edited by the user in the Submit window.
 	FText UserEditChangelistDescription = ChangelistDescriptionToSubmit;
 
+	// first check if there is a submit override bound
+	if (ISourceControlWindowsModule::Get().SubmitOverrideDelegate.IsBound())
+	{
+		// save the changelist description from the widget to perforce
+		const FString Identifier = ChangelistState->GetChangelist()->GetIdentifier();
+
+		SSubmitOverrideParameters SubmitOverrideParameters;
+		SubmitOverrideParameters.Description = UserEditChangelistDescription.ToString();
+		SubmitOverrideParameters.ToSubmit.SetSubtype<FString>(Identifier);
+
+		FSubmitOverrideReply SubmitOverrideReply = ISourceControlWindowsModule::Get().SubmitOverrideDelegate.Execute(SubmitOverrideParameters);
+
+		switch (SubmitOverrideReply)
+		{
+			//////////////////////////////////////////////////////////
+			case FSubmitOverrideReply::Handled:
+			{
+				FNotificationInfo Info(LOCTEXT("SCC_Checkin_SubmitOverride_Succeeded", "Successfully invoke the submit override!"));
+
+				Info.Text = LOCTEXT("SCC_Checkin_SubmitOverride_Succeeded", "Successfully invoked the submit override!");
+				Info.ExpireDuration = 8.0f;
+				Info.HyperlinkText = LOCTEXT("SCC_Checkin_ShowLog", "Show Message Log");
+				Info.Hyperlink = FSimpleDelegate::CreateLambda([]() { FMessageLog("SourceControl").Open(EMessageSeverity::Warning, true); });
+
+				TSharedPtr<SNotificationItem> Notification = FSlateNotificationManager::Get().AddNotification(Info);
+				Notification->SetCompletionState(SNotificationItem::CS_Success);
+				return;
+			}
+			
+			//////////////////////////////////////////////////////////
+			case FSubmitOverrideReply::Error:
+			{
+				FNotificationInfo Info(LOCTEXT("SCC_Checkin_SubmitOverride_Failed", "Failed to invoke the submit override!"));
+
+				Info.Text = LOCTEXT("SCC_Checkin_SubmitOverride_Failed", "Failed to invoke the submit override!");
+				Info.ExpireDuration = 8.0f;
+				Info.HyperlinkText = LOCTEXT("SCC_Checkin_ShowLog", "Show Message Log");
+				Info.Hyperlink = FSimpleDelegate::CreateLambda([]() { FMessageLog("SourceControl").Open(EMessageSeverity::Error, true); });
+
+				TSharedPtr<SNotificationItem> Notification = FSlateNotificationManager::Get().AddNotification(Info);
+				Notification->SetCompletionState(SNotificationItem::CS_Fail);
+				return;
+			}
+
+			//////////////////////////////////////////////////////////
+			case FSubmitOverrideReply::ProviderNotSupported:
+			default:
+				// continue default flow
+				break;
+		}
+	}
+
 	TSharedRef<SWindow> NewWindow = SNew(SWindow)
 		.Title(NSLOCTEXT("SourceControl.ConfirmSubmit", "Title", "Confirm changelist submit"))
 		.SizingRule(ESizingRule::UserSized)
diff --git a/Engine/Source/Editor/SourceControlWindows/Private/SourceControlWindows.cpp b/Engine/Source/Editor/SourceControlWindows/Private/SourceControlWindows.cpp
--- a/Engine/Source/Editor/SourceControlWindows/Private/SourceControlWindows.cpp
+++ b/Engine/Source/Editor/SourceControlWindows/Private/SourceControlWindows.cpp
@@ -5,6 +5,7 @@
 #include "AssetViewUtils.h"
 #include "FileHelpers.h"
 #include "ISourceControlModule.h"
+#include "ISourceControlWindowsModule.h"
 #include "SourceControlHelpers.h"
 #include "SourceControlOperations.h"
 #include "Framework/Application/SlateApplication.h"
@@ -407,6 +408,40 @@ bool FSourceControlWindows::PromptForCheckin(FCheckinResultInfo& OutResultInfo,
 		return bSuccess;
 	}
 
+	// first check if there is a submit override bound
+	if (ISourceControlWindowsModule::Get().SubmitOverrideDelegate.IsBound())
+	{
+		SSubmitOverrideParameters SubmitOverrideParameters;
+		SubmitOverrideParameters.Description = Description.Description.ToString();
+		SubmitOverrideParameters.ToSubmit.SetSubtype<TArray<FString>>(CombinedFileList);
+
+		FSubmitOverrideReply SubmitOverrideReply = ISourceControlWindowsModule::Get().SubmitOverrideDelegate.Execute(SubmitOverrideParameters);
+		switch (SubmitOverrideReply)
+		{
+			//////////////////////////////////////////////////////////
+			case FSubmitOverrideReply::Handled:
+			{
+				OutResultInfo.Result = ECommandResult::Succeeded;
+				OutResultInfo.Description = LOCTEXT("SCC_Checkin_SubmitOverride_Succeeded", "Successfully invoked the submit override!");
+				return true;
+			}
+
+			//////////////////////////////////////////////////////////
+			case FSubmitOverrideReply::Error:
+			{
+				OutResultInfo.Result = ECommandResult::Failed;
+				OutResultInfo.Description = LOCTEXT("SCC_Checkin_SubmitOverride_Failed", "Failed to invoke the submit override!");
+				return false;
+			}
+			
+			//////////////////////////////////////////////////////////
+			case FSubmitOverrideReply::ProviderNotSupported:
+			default:
+				// continue default flow
+				break;
+		}
+	}
+
 	FText VirtualizationFailureMsg;
 	if (!TryToVirtualizeFilesToSubmit(CombinedFileList, Description.Description, VirtualizationFailureMsg))
 	{
diff --git a/Engine/Source/Editor/SourceControlWindows/Public/ISourceControlWindowsModule.h b/Engine/Source/Editor/SourceControlWindows/Public/ISourceControlWindowsModule.h
--- a/Engine/Source/Editor/SourceControlWindows/Public/ISourceControlWindowsModule.h
+++ b/Engine/Source/Editor/SourceControlWindows/Public/ISourceControlWindowsModule.h
@@ -2,12 +2,32 @@
 #pragma once
 
 #include "CoreMinimal.h"
+#include "Containers/Union.h"
+
 #include "Modules/ModuleInterface.h"
 #include "Modules/ModuleManager.h"
 
 /**
  * SourceControlWindows module interface
  */
+
+enum class FSubmitOverrideReply
+{
+	Handled,
+	ProviderNotSupported,
+	Error,
+
+	Num
+};
+
+struct SSubmitOverrideParameters
+{
+	TUnion<FString, TArray<FString>> ToSubmit;
+	FString Description;
+};
+
+DECLARE_DELEGATE_RetVal_OneParam(FSubmitOverrideReply, FSubmitOverrideDelegate, SSubmitOverrideParameters /*InParameters*/);
+
 class ISourceControlWindowsModule : public IModuleInterface
 {
 public:
@@ -30,4 +50,6 @@ public:
 
 	DECLARE_EVENT_OneParam(ISourceControlWindowsModule, FChangelistFileDoubleClickedEvent, const FString&);
 	virtual FChangelistFileDoubleClickedEvent& OnChangelistFileDoubleClicked() = 0;
-};
+
+	FSubmitOverrideDelegate SubmitOverrideDelegate;
+};
\ No newline at end of file
