Fixing LOCTEXT conflicts introduced with previous CL 28533347.

diff --git a/Engine/Source/Editor/SourceControlWindows/Private/SSourceControlChangelists.cpp b/Engine/Source/Editor/SourceControlWindows/Private/SSourceControlChangelists.cpp
--- a/Engine/Source/Editor/SourceControlWindows/Private/SSourceControlChangelists.cpp
+++ b/Engine/Source/Editor/SourceControlWindows/Private/SSourceControlChangelists.cpp
@@ -2371,7 +2371,7 @@ void SSourceControlChangelistsWidget::OnSubmitChangelist()
 			//////////////////////////////////////////////////////////
 			case FSubmitOverrideReply::Handled:
 			{
-				FNotificationInfo Info(LOCTEXT("SCC_Checkin_SubmitOverride_Succeeded", "Successfully invoke the submit override!"));
+				FNotificationInfo Info(LOCTEXT("SCC_Checkin_SubmitOverride_Succeeded", "Successfully invoked the submit override!"));
 
 				Info.Text = LOCTEXT("SCC_Checkin_SubmitOverride_Succeeded", "Successfully invoked the submit override!");
 				Info.ExpireDuration = 8.0f;
