Restore validation outputs (without message log) to GetChangelistValidationResult

diff --git a/Engine/Source/Editor/SourceControlWindows/Private/SSourceControlChangelists.cpp b/Engine/Source/Editor/SourceControlWindows/Private/SSourceControlChangelists.cpp
--- a/Engine/Source/Editor/SourceControlWindows/Private/SSourceControlChangelists.cpp
+++ b/Engine/Source/Editor/SourceControlWindows/Private/SSourceControlChangelists.cpp
@@ -2226,7 +2226,56 @@ static bool GetChangelistValidationResult(FSourceControlChangelistPtr InChangeli
 
 	bool bValidationResult = true;
 
-	ValidationDelegate.ExecuteIfBound(InChangelist, ValidationResult, ValidationErrors, ValidationWarnings);
+	if (ValidationDelegate.ExecuteIfBound(InChangelist, ValidationResult, ValidationErrors, ValidationWarnings))
+	{
+		EMessageSeverity::Type MessageSeverity = EMessageSeverity::Info;
+
+		if (ValidationResult == EDataValidationResult::Invalid || ValidationErrors.Num() > 0)
+		{
+			OutValidationTitleText = LOCTEXT("SourceControl.Submit.ChangelistValidationError", "Changelist validation failed!").ToString();
+			bValidationResult = false;
+			MessageSeverity = EMessageSeverity::Error;
+		}
+		else if (ValidationResult == EDataValidationResult::NotValidated || ValidationWarnings.Num() > 0)
+		{
+			OutValidationTitleText = LOCTEXT("SourceControl.Submit.ChangelistValidationWarning", "Changelist validation has warnings!").ToString();
+			MessageSeverity = EMessageSeverity::Warning;
+		}
+		else
+		{
+			OutValidationTitleText = LOCTEXT("SourceControl.Submit.ChangelistValidationSuccess", "Changelist validation successful!").ToString();
+		}
+
+		auto AppendInfo = [](const TArray<FText>& Info, const FString& InfoType, FString& OutText)
+		{
+			const int32 MaxNumLinesDisplayed = 5;
+			int32 NumLinesDisplayed = 0;
+
+			if (Info.Num() > 0)
+			{
+				OutText += LINE_TERMINATOR;
+				OutText += FString::Printf(TEXT("Encountered %d %s:"), Info.Num(), *InfoType);
+
+				for (const FText& Line : Info)
+				{
+					if (NumLinesDisplayed >= MaxNumLinesDisplayed)
+					{
+						OutText += LINE_TERMINATOR;
+						OutText += FString::Printf(TEXT("See log for complete list of %s"), *InfoType);
+						break;
+					}
+
+					OutText += LINE_TERMINATOR;
+					OutText += Line.ToString();
+
+					++NumLinesDisplayed;
+				}
+			}
+		};
+
+		AppendInfo(ValidationErrors, TEXT("errors"), OutValidationErrorsText);
+		AppendInfo(ValidationWarnings, TEXT("warnings"), OutValidationWarningsText);
+	}
 
 	return bValidationResult;
 }
