- Moved the delegate before the content valdiation
- Added logic to remove the validation tag if the content validation fails

diff --git a/Engine/Source/Editor/SourceControlWindows/Private/SSourceControlChangelists.cpp b/Engine/Source/Editor/SourceControlWindows/Private/SSourceControlChangelists.cpp
--- a/Engine/Source/Editor/SourceControlWindows/Private/SSourceControlChangelists.cpp
+++ b/Engine/Source/Editor/SourceControlWindows/Private/SSourceControlChangelists.cpp
@@ -104,6 +104,15 @@ FText UpdateChangelistDescriptionToSubmitIfNeeded(const bool bInValidationResult
 		return DescriptionString.Find(ValidationString) != INDEX_NONE;
 	};
 
+	auto RemoveValidationFlag = [&GetChangelistValidationTag](const FText& InChangelistDescription)
+	{
+		FString DescriptionString = InChangelistDescription.ToString();
+		FString ValidationString = GetChangelistValidationTag().ToString();
+		DescriptionString.ReplaceInline(*ValidationString, TEXT(""));
+
+		return DescriptionString;
+	};
+
 	if (bInValidationResult && USourceControlPreferences::IsValidationTagEnabled() && !ContainsValidationFlag(InChangelistDescription))
 	{
 		FStringOutputDevice Str;
@@ -115,6 +124,12 @@ FText UpdateChangelistDescriptionToSubmitIfNeeded(const bool bInValidationResult
 		return FText::FromString(Str);
 	}
 
+	if (!bInValidationResult && USourceControlPreferences::IsValidationTagEnabled() && ContainsValidationFlag(InChangelistDescription))
+	{
+		FString NewChangelistDescription = RemoveValidationFlag(InChangelistDescription);
+		return FText::FromString(NewChangelistDescription);
+	}
+
 	return InChangelistDescription;
 }
 
@@ -2319,21 +2334,6 @@ void SSourceControlChangelistsWidget::OnSubmitChangelist()
 		return;
 	}
 
-	FString ChangelistValidationTitle;
-	FString ChangelistValidationWarningsText;
-	FString ChangelistValidationErrorsText;
-	bool bValidationResult = GetChangelistValidationResult(ChangelistState->GetChangelist(), ChangelistValidationTitle, ChangelistValidationWarningsText, ChangelistValidationErrorsText);
-
-	// The description from the source control.
-	const FText CurrentChangelistDescription = ChangelistState->GetDescriptionText();
-	const bool bAskForChangelistDescription = (CurrentChangelistDescription.IsEmptyOrWhitespace());
-
-	// The description possibly updated with the #validated proposed to the user.
-	FText ChangelistDescriptionToSubmit = UpdateChangelistDescriptionToSubmitIfNeeded(bValidationResult, CurrentChangelistDescription);
-
-	// The description once edited by the user in the Submit window.
-	FText UserEditChangelistDescription = ChangelistDescriptionToSubmit;
-
 	// first check if there is a submit override bound
 	if (ISourceControlWindowsModule::Get().SubmitOverrideDelegate.IsBound())
 	{
@@ -2341,7 +2341,7 @@ void SSourceControlChangelistsWidget::OnSubmitChangelist()
 		const FString Identifier = ChangelistState->GetChangelist()->GetIdentifier();
 
 		SSubmitOverrideParameters SubmitOverrideParameters;
-		SubmitOverrideParameters.Description = UserEditChangelistDescription.ToString();
+		SubmitOverrideParameters.Description = ChangelistState->GetDescriptionText().ToString();
 		SubmitOverrideParameters.ToSubmit.SetSubtype<FString>(Identifier);
 
 		FSubmitOverrideReply SubmitOverrideReply = ISourceControlWindowsModule::Get().SubmitOverrideDelegate.Execute(SubmitOverrideParameters);
@@ -2386,6 +2386,21 @@ void SSourceControlChangelistsWidget::OnSubmitChangelist()
 		}
 	}
 
+	FString ChangelistValidationTitle;
+	FString ChangelistValidationWarningsText;
+	FString ChangelistValidationErrorsText;
+	bool bValidationResult = GetChangelistValidationResult(ChangelistState->GetChangelist(), ChangelistValidationTitle, ChangelistValidationWarningsText, ChangelistValidationErrorsText);
+
+	// The description from the source control.
+	const FText CurrentChangelistDescription = ChangelistState->GetDescriptionText();
+	const bool bAskForChangelistDescription = (CurrentChangelistDescription.IsEmptyOrWhitespace());
+
+	// The description possibly updated with the #validated proposed to the user.
+	FText ChangelistDescriptionToSubmit = UpdateChangelistDescriptionToSubmitIfNeeded(bValidationResult, CurrentChangelistDescription);
+
+	// The description once edited by the user in the Submit window.
+	FText UserEditChangelistDescription = ChangelistDescriptionToSubmit;
+
 	TSharedRef<SWindow> NewWindow = SNew(SWindow)
 		.Title(NSLOCTEXT("SourceControl.ConfirmSubmit", "Title", "Confirm changelist submit"))
 		.SizingRule(ESizingRule::UserSized)
