Refresh the UI once the delegate has been called.

diff --git a/Engine/Source/Editor/SourceControlWindows/Private/SSourceControlChangelists.cpp b/Engine/Source/Editor/SourceControlWindows/Private/SSourceControlChangelists.cpp
--- a/Engine/Source/Editor/SourceControlWindows/Private/SSourceControlChangelists.cpp
+++ b/Engine/Source/Editor/SourceControlWindows/Private/SSourceControlChangelists.cpp
@@ -2360,6 +2360,8 @@ void SSourceControlChangelistsWidget::OnSubmitChangelist()
 
 				TSharedPtr<SNotificationItem> Notification = FSlateNotificationManager::Get().AddNotification(Info);
 				Notification->SetCompletionState(SNotificationItem::CS_Success);
+
+				this->OnRefreshUI(ERefreshFlags::SourceControlChangelists);
 				return;
 			}
 			
@@ -2375,6 +2377,8 @@ void SSourceControlChangelistsWidget::OnSubmitChangelist()
 
 				TSharedPtr<SNotificationItem> Notification = FSlateNotificationManager::Get().AddNotification(Info);
 				Notification->SetCompletionState(SNotificationItem::CS_Fail);
+
+				this->OnRefreshUI(ERefreshFlags::SourceControlChangelists);
 				return;
 			}
 
