Optimize source control changelist shelved state update.

Results with a 40K actors changelist editor startup time:
  - Before: FPerforce::AddShelvedFilesToChangelist took 2:35 min
  - After: FPerforce::AddShelvedFilesToChangelist took 160.08 ms

diff --git a/Engine/Plugins/Developer/Concert/ConcertSync/ConcertSyncClient/Source/ConcertSyncClient/Private/ConcertSourceControlProxy.cpp b/Engine/Plugins/Developer/Concert/ConcertSync/ConcertSyncClient/Source/ConcertSyncClient/Private/ConcertSourceControlProxy.cpp
--- a/Engine/Plugins/Developer/Concert/ConcertSync/ConcertSyncClient/Source/ConcertSyncClient/Private/ConcertSourceControlProxy.cpp
+++ b/Engine/Plugins/Developer/Concert/ConcertSync/ConcertSyncClient/Source/ConcertSyncClient/Private/ConcertSourceControlProxy.cpp
@@ -238,16 +238,24 @@ const FDateTime& FConcertSourceControlChangelistStateProxy::GetTimeStamp() const
 	return ActualState.IsValid() ? ActualState->GetTimeStamp() : DummyTime;
 }
 
-const TArray<FSourceControlStateRef>& FConcertSourceControlChangelistStateProxy::GetFilesStates() const
+const TArray<FSourceControlStateRef> FConcertSourceControlChangelistStateProxy::GetFilesStates() const
 {
-	static const TArray<FSourceControlStateRef> DummyResult;
-	return ActualState.IsValid() ? ActualState->GetFilesStates() : DummyResult;
+	return ActualState.IsValid() ? ActualState->GetFilesStates() : TArray<FSourceControlStateRef>();
 }
 
-const TArray<FSourceControlStateRef>& FConcertSourceControlChangelistStateProxy::GetShelvedFilesStates() const
+int32 FConcertSourceControlChangelistStateProxy::GetFilesStatesNum() const
 {
-	static const TArray<FSourceControlStateRef> DummyResult;
-	return ActualState.IsValid() ? ActualState->GetShelvedFilesStates() : DummyResult;
+	return ActualState.IsValid() ? ActualState->GetFilesStates().Num() : 0;
+}
+
+const TArray<FSourceControlStateRef> FConcertSourceControlChangelistStateProxy::GetShelvedFilesStates() const
+{
+	return ActualState.IsValid() ? ActualState->GetShelvedFilesStates() : TArray<FSourceControlStateRef>();
+}
+
+int32 FConcertSourceControlChangelistStateProxy::GetShelvedFilesStatesNum() const
+{
+	return ActualState.IsValid() ? ActualState->GetShelvedFilesStates().Num() : 0;
 }
 
 FSourceControlChangelistRef FConcertSourceControlChangelistStateProxy::GetChangelist() const
diff --git a/Engine/Plugins/Developer/Concert/ConcertSync/ConcertSyncClient/Source/ConcertSyncClient/Private/ConcertSourceControlProxy.h b/Engine/Plugins/Developer/Concert/ConcertSync/ConcertSyncClient/Source/ConcertSyncClient/Private/ConcertSourceControlProxy.h
--- a/Engine/Plugins/Developer/Concert/ConcertSync/ConcertSyncClient/Source/ConcertSyncClient/Private/ConcertSourceControlProxy.h
+++ b/Engine/Plugins/Developer/Concert/ConcertSync/ConcertSyncClient/Source/ConcertSyncClient/Private/ConcertSourceControlProxy.h
@@ -78,8 +78,10 @@ public:
 	virtual FText GetDescriptionText() const override;
 	virtual FText GetDisplayTooltip() const override;
 	virtual const FDateTime& GetTimeStamp() const override;
-	virtual const TArray<FSourceControlStateRef>& GetFilesStates() const override;
-	virtual const TArray<FSourceControlStateRef>& GetShelvedFilesStates() const override;
+	virtual const TArray<FSourceControlStateRef> GetFilesStates() const override;
+	virtual int32 GetFilesStatesNum() const override;
+	virtual const TArray<FSourceControlStateRef> GetShelvedFilesStates() const override;
+	virtual int32 GetShelvedFilesStatesNum() const override;
 	virtual FSourceControlChangelistRef GetChangelist() const override;
 
 private:
diff --git a/Engine/Plugins/Developer/PerforceSourceControl/Source/PerforceSourceControl/Private/PerforceSourceControlChangelistState.cpp b/Engine/Plugins/Developer/PerforceSourceControl/Source/PerforceSourceControl/Private/PerforceSourceControlChangelistState.cpp
--- a/Engine/Plugins/Developer/PerforceSourceControl/Source/PerforceSourceControl/Private/PerforceSourceControlChangelistState.cpp
+++ b/Engine/Plugins/Developer/PerforceSourceControl/Source/PerforceSourceControl/Private/PerforceSourceControlChangelistState.cpp
@@ -41,14 +41,26 @@ const FDateTime& FPerforceSourceControlChangelistState::GetTimeStamp() const
 	return TimeStamp;
 }
 
-const TArray<FSourceControlStateRef>& FPerforceSourceControlChangelistState::GetFilesStates() const
+const TArray<FSourceControlStateRef> FPerforceSourceControlChangelistState::GetFilesStates() const
 {
 	return Files;
 }
 
-const TArray<FSourceControlStateRef>& FPerforceSourceControlChangelistState::GetShelvedFilesStates() const
+int32 FPerforceSourceControlChangelistState::GetFilesStatesNum() const
 {
-	return ShelvedFiles;
+	return Files.Num();
+}
+
+const TArray<FSourceControlStateRef> FPerforceSourceControlChangelistState::GetShelvedFilesStates() const
+{
+	TArray<FSourceControlStateRef> Result;
+	ShelvedFiles.GenerateValueArray(Result);
+	return Result;
+}
+
+int32 FPerforceSourceControlChangelistState::GetShelvedFilesStatesNum() const
+{
+	return ShelvedFiles.Num();
 }
 
 FSourceControlChangelistRef FPerforceSourceControlChangelistState::GetChangelist() const
diff --git a/Engine/Plugins/Developer/PerforceSourceControl/Source/PerforceSourceControl/Private/PerforceSourceControlChangelistState.h b/Engine/Plugins/Developer/PerforceSourceControl/Source/PerforceSourceControl/Private/PerforceSourceControlChangelistState.h
--- a/Engine/Plugins/Developer/PerforceSourceControl/Source/PerforceSourceControl/Private/PerforceSourceControlChangelistState.h
+++ b/Engine/Plugins/Developer/PerforceSourceControl/Source/PerforceSourceControl/Private/PerforceSourceControlChangelistState.h
@@ -55,9 +55,11 @@ public:
 	 */
 	virtual const FDateTime& GetTimeStamp() const override;
 
-	virtual const TArray<FSourceControlStateRef>& GetFilesStates() const override;
+	virtual const TArray<FSourceControlStateRef> GetFilesStates() const override;
+	virtual int32 GetFilesStatesNum() const override;
 
-	virtual const TArray<FSourceControlStateRef>& GetShelvedFilesStates() const override;
+	virtual const TArray<FSourceControlStateRef> GetShelvedFilesStates() const override;
+	virtual int32 GetShelvedFilesStatesNum() const override;
 
 	virtual FSourceControlChangelistRef GetChangelist() const override;
 
@@ -68,7 +70,7 @@ public:
 	bool bHasShelvedFiles;
 
 	TArray<FSourceControlStateRef> Files;
-	TArray<FSourceControlStateRef> ShelvedFiles;
+	TMap<FString, FSourceControlStateRef> ShelvedFiles;
 
 	/** The timestamp of the last update */
 	FDateTime TimeStamp;
diff --git a/Engine/Plugins/Developer/PerforceSourceControl/Source/PerforceSourceControl/Private/PerforceSourceControlOperations.cpp b/Engine/Plugins/Developer/PerforceSourceControl/Source/PerforceSourceControl/Private/PerforceSourceControlOperations.cpp
--- a/Engine/Plugins/Developer/PerforceSourceControl/Source/PerforceSourceControl/Private/PerforceSourceControlOperations.cpp
+++ b/Engine/Plugins/Developer/PerforceSourceControl/Source/PerforceSourceControl/Private/PerforceSourceControlOperations.cpp
@@ -2395,11 +2395,9 @@ static bool AddShelvedFilesToChangelist(FPerforceSourceControlProvider& SCCProvi
 			ItFilename = DepotToFileMap[It.Key()];
 		}
 
-		int32 Index = Algo::IndexOfByPredicate(ChangelistState->ShelvedFiles, [&ItFilename](const FSourceControlStateRef& ShelvedFile) {
-				return ShelvedFile->GetFilename() == ItFilename;
-			});
+		FSourceControlStateRef* SourceControlStateRef = ChangelistState->ShelvedFiles.Find(ItFilename);
 
-		if (Index < 0)
+		if (!SourceControlStateRef)
 		{
 			// Create new entry
 			TSharedRef<FPerforceSourceControlState, ESPMode::ThreadSafe> ShelvedFileState = MakeShareable(new FPerforceSourceControlState(ItFilename));
@@ -2417,11 +2415,10 @@ static bool AddShelvedFilesToChangelist(FPerforceSourceControlProvider& SCCProvi
 			}
 
 			// Add to shelved files
-			Index = ChangelistState->ShelvedFiles.Num();
-			ChangelistState->ShelvedFiles.Add(ShelvedFileState);
+			SourceControlStateRef = &ChangelistState->ShelvedFiles.Add(ItFilename, ShelvedFileState);
 		}
 
-		TSharedRef<FPerforceSourceControlState, ESPMode::ThreadSafe> FileState = StaticCastSharedRef<FPerforceSourceControlState>(ChangelistState->ShelvedFiles[Index]);
+		TSharedRef<FPerforceSourceControlState, ESPMode::ThreadSafe> FileState = StaticCastSharedRef<FPerforceSourceControlState>(*SourceControlStateRef);
 
  		FileState->SetState(It.Value());
  		FileState->TimeStamp = Now;
@@ -2472,7 +2469,8 @@ bool FPerforceGetPendingChangelistsWorker::UpdateStates() const
 		bool bUpdateShelvedFiles = (OutCLShelvedFilesStates.Num() == OutChangelistsStates.Num());
 		if(bUpdateShelvedFiles)
 		{
-			ChangelistState->ShelvedFiles.Reset(OutCLShelvedFilesStates[StatusIndex].Num());
+			ChangelistState->ShelvedFiles.Reset();
+			ChangelistState->ShelvedFiles.Reserve(OutCLShelvedFilesStates[StatusIndex].Num());
 			AddShelvedFilesToChangelist(GetSCCProvider(), OutCLShelvedFilesStates[StatusIndex], OutCLShelvedFilesMap[StatusIndex], ChangelistState, &Now);
 		}
 	}
@@ -3136,24 +3134,20 @@ bool FPerforceDeleteShelveWorker::UpdateStates() const
 
 		if (FilesToRemove.Num() > 0)
 		{
-			return ChangelistState->ShelvedFiles.RemoveAll([this](FSourceControlStateRef& State) -> bool
-				{
-					return Algo::AnyOf(FilesToRemove, [&State](auto& File) {
-						return State->GetFilename() == File;
-					});
-				}) > 0;
+			bool bRemovedFiles = false;
+			for (const FString& FileToRemove : FilesToRemove)
+			{
+				bRemovedFiles |= ChangelistState->ShelvedFiles.Remove(FileToRemove) > 0;
+			}
+			return bRemovedFiles;
 		}
-		else
-		{
-			bool bHadShelvedFiles = (ChangelistState->ShelvedFiles.Num() > 0);
-			ChangelistState->ShelvedFiles.Reset();
-			return bHadShelvedFiles;
-		}		
-	}
-	else
-	{
-		return false;
+
+		bool bHadShelvedFiles = (ChangelistState->ShelvedFiles.Num() > 0);
+		ChangelistState->ShelvedFiles.Reset();
+		return bHadShelvedFiles;
 	}
+
+	return false;
 }
 
 
diff --git a/Engine/Plugins/Developer/PlasticSourceControl/Source/PlasticSourceControl/Private/PlasticSourceControlChangelistState.cpp b/Engine/Plugins/Developer/PlasticSourceControl/Source/PlasticSourceControl/Private/PlasticSourceControlChangelistState.cpp
--- a/Engine/Plugins/Developer/PlasticSourceControl/Source/PlasticSourceControl/Private/PlasticSourceControlChangelistState.cpp
+++ b/Engine/Plugins/Developer/PlasticSourceControl/Source/PlasticSourceControl/Private/PlasticSourceControlChangelistState.cpp
@@ -35,16 +35,26 @@ const FDateTime& FPlasticSourceControlChangelistState::GetTimeStamp() const
 	return TimeStamp;
 }
 
-const TArray<FSourceControlStateRef>& FPlasticSourceControlChangelistState::GetFilesStates() const
+const TArray<FSourceControlStateRef> FPlasticSourceControlChangelistState::GetFilesStates() const
 {
 	return Files;
 }
 
-const TArray<FSourceControlStateRef>& FPlasticSourceControlChangelistState::GetShelvedFilesStates() const
+int32 FPlasticSourceControlChangelistState::GetFilesStatesNum() const
+{
+	return Files.Num();
+}
+
+const TArray<FSourceControlStateRef> FPlasticSourceControlChangelistState::GetShelvedFilesStates() const
 {
 	return ShelvedFiles;
 }
 
+int32 FPlasticSourceControlChangelistState::GetShelvedFilesStatesNum() const
+{
+	return ShelvedFiles.Num();
+}
+
 FSourceControlChangelistRef FPlasticSourceControlChangelistState::GetChangelist() const
 {
 	FPlasticSourceControlChangelistRef ChangelistCopy = MakeShareable( new FPlasticSourceControlChangelist(Changelist));
diff --git a/Engine/Plugins/Developer/PlasticSourceControl/Source/PlasticSourceControl/Private/PlasticSourceControlChangelistState.h b/Engine/Plugins/Developer/PlasticSourceControl/Source/PlasticSourceControl/Private/PlasticSourceControlChangelistState.h
--- a/Engine/Plugins/Developer/PlasticSourceControl/Source/PlasticSourceControl/Private/PlasticSourceControlChangelistState.h
+++ b/Engine/Plugins/Developer/PlasticSourceControl/Source/PlasticSourceControl/Private/PlasticSourceControlChangelistState.h
@@ -62,9 +62,11 @@ public:
 	 */
 	virtual const FDateTime& GetTimeStamp() const override;
 
-	virtual const TArray<FSourceControlStateRef>& GetFilesStates() const override;
+	virtual const TArray<FSourceControlStateRef> GetFilesStates() const override;
+	virtual int32 GetFilesStatesNum() const override;
 
-	virtual const TArray<FSourceControlStateRef>& GetShelvedFilesStates() const override;
+	virtual const TArray<FSourceControlStateRef> GetShelvedFilesStates() const override;
+	virtual int32 GetShelvedFilesStatesNum() const override;
 
 	virtual FSourceControlChangelistRef GetChangelist() const override;
 
diff --git a/Engine/Source/Developer/SourceControl/Public/ISourceControlChangelistState.h b/Engine/Source/Developer/SourceControl/Public/ISourceControlChangelistState.h
--- a/Engine/Source/Developer/SourceControl/Public/ISourceControlChangelistState.h
+++ b/Engine/Source/Developer/SourceControl/Public/ISourceControlChangelistState.h
@@ -64,9 +64,11 @@ public:
 	 */
 	virtual const FDateTime& GetTimeStamp() const = 0;
 
-	virtual const TArray<FSourceControlStateRef>& GetFilesStates() const = 0;
+	virtual const TArray<FSourceControlStateRef> GetFilesStates() const = 0;
+	virtual int32 GetFilesStatesNum() const = 0;
 
-	virtual const TArray<FSourceControlStateRef>& GetShelvedFilesStates() const = 0;
+	virtual const TArray<FSourceControlStateRef> GetShelvedFilesStates() const = 0;
+	virtual int32 GetShelvedFilesStatesNum() const = 0;
 
 	/**
 	 * Returns the object on which this state was constructed
diff --git a/Engine/Source/Editor/SourceControlWindows/Private/SSourceControlChangelists.cpp b/Engine/Source/Editor/SourceControlWindows/Private/SSourceControlChangelists.cpp
--- a/Engine/Source/Editor/SourceControlWindows/Private/SSourceControlChangelists.cpp
+++ b/Engine/Source/Editor/SourceControlWindows/Private/SSourceControlChangelists.cpp
@@ -1082,7 +1082,7 @@ void SSourceControlChangelistsWidget::OnRefreshSourceControlWidgets(int64 CurrUp
 		}
 		ShelvedChangelistTreeItem = ChangelistTreeItem->ShelvedChangelistItem;
 
-		const bool bShelvedChangelistPassesFilter = ChangelistState->GetShelvedFilesStates().Num() > 0 && ChangelistTextFilter->PassesFilter(*ShelvedChangelistTreeItem);
+		const bool bShelvedChangelistPassesFilter = ChangelistState->GetShelvedFilesStatesNum() > 0 && ChangelistTextFilter->PassesFilter(*ShelvedChangelistTreeItem);
 		const bool bShelvedChangelistMustBeDisplayed = bShelvedChangelistPassesFilter && !ChangelistTextFilter->GetRawFilterText().IsEmpty();
 		const bool bChangelistPassedFilter = bShelvedChangelistMustBeDisplayed || ChangelistTextFilter->PassesFilter(*ChangelistTreeItem);
 		if (bChangelistPassedFilter)
@@ -1990,7 +1990,7 @@ bool SSourceControlChangelistsWidget::CanDeleteChangelist(FText* OutFailureMessa
 		}
 		return false;
 	}
-	else if (ChangelistState->GetFilesStates().Num() > 0 || ChangelistState->GetShelvedFilesStates().Num() > 0)
+	else if (ChangelistState->GetFilesStatesNum() > 0 || ChangelistState->GetShelvedFilesStatesNum() > 0)
 	{
 		if (OutFailureMessage)
 		{
@@ -2046,7 +2046,7 @@ void SSourceControlChangelistsWidget::OnRevertUnchanged()
 
 bool SSourceControlChangelistsWidget::CanRevertUnchanged()
 {
-	return HasFilesSelected() || (GetCurrentChangelistState() && GetCurrentChangelistState()->GetFilesStates().Num() > 0);
+	return HasFilesSelected() || (GetCurrentChangelistState() && GetCurrentChangelistState()->GetFilesStatesNum() > 0);
 }
 
 void SSourceControlChangelistsWidget::OnRevert()
@@ -2151,7 +2151,7 @@ bool SSourceControlChangelistsWidget::CanRevert()
 	FUncontrolledChangelistStatePtr CurrentUncontrolledChangelistState = GetCurrentUncontrolledChangelistState();
 
 	return HasFilesSelected()
-		|| (CurrentChangelistState.IsValid() && CurrentChangelistState->GetFilesStates().Num() > 0)
+		|| (CurrentChangelistState.IsValid() && CurrentChangelistState->GetFilesStatesNum() > 0)
 		|| (CurrentUncontrolledChangelistState.IsValid() && CurrentUncontrolledChangelistState->GetFilesStates().Num() > 0);
 }
 
@@ -2511,7 +2511,7 @@ bool SSourceControlChangelistsWidget::CanSubmitChangelist(FText* OutFailureMessa
 		}
 		return false;
 	}
-	else if (Changelist->GetFilesStates().Num() <= 0)
+	else if (Changelist->GetFilesStatesNum() <= 0)
 	{
 		if (OutFailureMessage)
 		{
@@ -2519,7 +2519,7 @@ bool SSourceControlChangelistsWidget::CanSubmitChangelist(FText* OutFailureMessa
 		}
 		return false;
 	}
-	else if (Changelist->GetShelvedFilesStates().Num() > 0)
+	else if (Changelist->GetShelvedFilesStatesNum() > 0)
 	{
 		if (OutFailureMessage)
 		{
@@ -2565,7 +2565,7 @@ void SSourceControlChangelistsWidget::OnValidateChangelist()
 bool SSourceControlChangelistsWidget::CanValidateChangelist()
 {
 	FSourceControlChangelistStatePtr Changelist = GetCurrentChangelistState();
-	return Changelist != nullptr && Changelist->GetFilesStates().Num() > 0;
+	return Changelist != nullptr && Changelist->GetFilesStatesNum() > 0;
 }
 
 void SSourceControlChangelistsWidget::OnNewUncontrolledChangelist()
@@ -3036,12 +3036,12 @@ TSharedPtr<SWidget> SSourceControlChangelistsWidget::OnOpenContextMenu()
 				FCanExecuteAction::CreateSP(this, &SSourceControlChangelistsWidget::CanRevert)));
 	}
 
-	if (bHasSelectedChangelist && (bHasSelectedFiles || bHasSelectedShelvedFiles || (bHasSelectedChangelist && (GetCurrentChangelistState()->GetFilesStates().Num() > 0 || GetCurrentChangelistState()->GetShelvedFilesStates().Num() > 0))))
+	if (bHasSelectedChangelist && (bHasSelectedFiles || bHasSelectedShelvedFiles || (bHasSelectedChangelist && (GetCurrentChangelistState()->GetFilesStatesNum() > 0 || GetCurrentChangelistState()->GetShelvedFilesStates().Num() > 0))))
 	{
 		Section.AddSeparator("ShelveSeparator");
 	}
 
-	if (bHasSelectedChangelist && (bHasSelectedFiles || (bHasSelectedChangelist && GetCurrentChangelistState()->GetFilesStates().Num() > 0)))
+	if (bHasSelectedChangelist && (bHasSelectedFiles || (bHasSelectedChangelist && GetCurrentChangelistState()->GetFilesStatesNum() > 0)))
 	{
 		Section.AddMenuEntry("Shelve",
 			LOCTEXT("SourceControl_Shelve", "Shelve Files"),
diff --git a/Engine/Source/Editor/SourceControlWindows/Private/SSourceControlCommon.h b/Engine/Source/Editor/SourceControlWindows/Private/SSourceControlCommon.h
--- a/Engine/Source/Editor/SourceControlWindows/Private/SSourceControlCommon.h
+++ b/Engine/Source/Editor/SourceControlWindows/Private/SSourceControlCommon.h
@@ -165,12 +165,12 @@ struct FChangelistTreeItem : public IChangelistTreeItem
 
 	int32 GetFileCount() const
 	{
-		return ChangelistState->GetFilesStates().Num();
+		return ChangelistState->GetFilesStatesNum();
 	}
 
 	int32 GetShelvedFileCount() const
 	{
-		return ChangelistState->GetShelvedFilesStates().Num();
+		return ChangelistState->GetShelvedFilesStatesNum();
 	}
 
 	TSharedRef<ISourceControlChangelistState> ChangelistState;
