Skein / ConflictUE: Streamlined behavior when user clicks CHECK-IN CHANGES.

Previously, the SYNC was done prior to showing the SUBMIT FILES dialog. If another user submitted anything while the user was typing their changelist description, the CHECK IN would fail (for Skein) because the user is no longer checking in as latest. By moving the SYNC to after the changelist description is completed, this is avoided. In addition, this allows the user once again to click the CHECK IN CHANGES without syncing to head if he just wants to see what local changes he has that need checking in - a workflow which I imagine is not uncommon.

Additionally, the check-in should be aborted if the SYNC results in the user entering conflict resolution.

#rnx

diff --git a/Engine/Source/Editor/SourceControlWindows/Private/SourceControlWindows.cpp b/Engine/Source/Editor/SourceControlWindows/Private/SourceControlWindows.cpp
--- a/Engine/Source/Editor/SourceControlWindows/Private/SourceControlWindows.cpp
+++ b/Engine/Source/Editor/SourceControlWindows/Private/SourceControlWindows.cpp
@@ -91,29 +91,6 @@ bool FSourceControlWindows::ChoosePackagesToCheckIn(const FSourceControlWindowsO
 		return false;
 	}
 
-	if (ISourceControlModule::Get().GetProvider().UsesSnapshots())
-	{
-		bool bSyncNeeded = FSourceControlWindows::CanSyncLatest();
-		if (bSyncNeeded)
-		{
-			FBookmarkScoped BookmarkScoped; // Preserve viewport camera orientation.
-
-			if (!FSourceControlWindows::SyncLatest())
-			{
-				FCheckinResultInfo ResultInfo;
-				ResultInfo.Description = LOCTEXT("SourceControlSyncFailed", "Revision control failed to sync to the latest revision.");
-				OnCompleteDelegate.ExecuteIfBound(ResultInfo);
-
-				return false;
-			}
-		}
-		else
-		{
-			// SyncLatest saves packages so do that explicitly if sync is not needed.
-			SaveDirtyPackages(/*bUseDialog=*/false);
-		}
-	}
-
 	// Start selection process...
 
 	// make sure we update the SCC status of all packages (this could take a long time, so we will run it as a background task)
@@ -337,6 +314,37 @@ bool FSourceControlWindows::PromptForCheckin(FCheckinResultInfo& OutResultInfo,
 		return false;
 	}
 
+	//~~~~~~~~~~~~~~~~~~~~~~~~~~~~~~~~~~~~~~~~~~
+	// Sync to latest if using snapshots.
+	if (ISourceControlModule::Get().GetProvider().UsesSnapshots())
+	{
+		const bool bSyncNeeded = FSourceControlWindows::CanSyncLatest();
+		if (bSyncNeeded)
+		{
+			FBookmarkScoped BookmarkScoped; // Preserve viewport camera orientation.
+
+			if (!FSourceControlWindows::SyncLatest())
+			{
+				OutResultInfo.Description = LOCTEXT("SCC_Checkin_Aborted_Sync", "File check in aborted because the sync to the latest snapshot failed.");
+				return false;
+			}
+
+			TArray<FSourceControlStateRef> Conflicts = ISourceControlModule::Get().GetProvider().GetCachedStateByPredicate(
+				[](const FSourceControlStateRef& State)
+				{
+					return State->IsConflicted();
+				}
+			);
+
+			const bool bConflictsRemaining = (Conflicts.Num() > 0);
+			if (bConflictsRemaining)
+			{
+				OutResultInfo.Description = LOCTEXT("SCC_Checkin_Aborted_Conflicts", "File check in aborted because the sync to the latest snapshot resulted in conflicts that need to be resolved.");
+				return false;
+			}
+		}
+	}
+
 
 	//~~~~~~~~~~~~~~~~~~~~~~~~~~~~~~~~~~~~~~~~~~
 	// Get description from the dialog
