Skein / ConflictUE: Fixed 'check in changes' no longer saving dirty packages since #29866499.

#rnx

diff --git a/Engine/Source/Editor/SourceControlWindows/Private/SourceControlWindows.cpp b/Engine/Source/Editor/SourceControlWindows/Private/SourceControlWindows.cpp
--- a/Engine/Source/Editor/SourceControlWindows/Private/SourceControlWindows.cpp
+++ b/Engine/Source/Editor/SourceControlWindows/Private/SourceControlWindows.cpp
@@ -91,6 +91,11 @@ bool FSourceControlWindows::ChoosePackagesToCheckIn(const FSourceControlWindowsO
 		return false;
 	}
 
+	if (ISourceControlModule::Get().GetProvider().UsesSnapshots())
+	{
+		SaveDirtyPackages(/*bUseDialog=*/false);
+	}
+
 	// Start selection process...
 
 	// make sure we update the SCC status of all packages (this could take a long time, so we will run it as a background task)
