Source Control: don't show the success icon when unable to submit the changelist and there's no warnings or errors.

diff --git a/Engine/Source/Editor/SourceControlWindows/Private/SSourceControlSubmit.cpp b/Engine/Source/Editor/SourceControlWindows/Private/SSourceControlSubmit.cpp
--- a/Engine/Source/Editor/SourceControlWindows/Private/SSourceControlSubmit.cpp
+++ b/Engine/Source/Editor/SourceControlWindows/Private/SSourceControlSubmit.cpp
@@ -283,9 +283,7 @@ void SSourceControlSubmitWidget::Construct(const FArguments& InArgs)
 		const FName ChangelistWarningsIconName = TEXT("Icons.WarningWithColor.Large");
 		const FName ChangelistErrorsIconName = TEXT("Icons.ErrorWithColor.Large");
 
-		const bool bIsSuccess = ChangelistResultWarningsText.IsEmpty() && ChangelistResultErrorsText.IsEmpty();
-
-		if (bIsSuccess)
+		if (bAllowSubmit)
 		{
 			Contents->AddSlot()
 			.AutoHeight()
@@ -310,6 +308,20 @@ void SSourceControlSubmitWidget::Construct(const FArguments& InArgs)
 		}
 		else
 		{
+			Contents->AddSlot()
+			.AutoHeight()
+			.Padding(FMargin(5))
+			[
+				SNew(SHorizontalBox)
+				+SHorizontalBox::Slot()
+				[
+					SNew(SMultiLineEditableTextBox)
+					.Text(FText::FromString(ChangelistResultText))
+					.AutoWrapText(true)
+					.IsReadOnly(true)
+				]
+			];
+
 			if (!ChangelistResultErrorsText.IsEmpty())
 			{
 				Contents->AddSlot()
