Updated display string for changelist descriptions to use space character in place of other whitespace characters (similar to P4V).
- allows viewing descriptions which traverse multiple lines
- only affects the single line view (tooltip, etc still display as before)
#rnx

diff --git a/Engine/Source/Editor/SourceControlWindows/Private/SSourceControlChangelistRows.cpp b/Engine/Source/Editor/SourceControlWindows/Private/SSourceControlChangelistRows.cpp
--- a/Engine/Source/Editor/SourceControlWindows/Private/SSourceControlChangelistRows.cpp
+++ b/Engine/Source/Editor/SourceControlWindows/Private/SSourceControlChangelistRows.cpp
@@ -169,7 +169,8 @@ FText SChangelistTableRow::GetChangelistDescriptionText() const
 
 FText SChangelistTableRow::GetChangelistDescriptionSingleLineText() const
 {
-	return SSourceControlCommon::GetSingleLineChangelistDescription(TreeItem->GetDescriptionText());
+	using namespace SSourceControlCommon;
+	return GetSingleLineChangelistDescription(TreeItem->GetDescriptionText(), ESingleLineFlags::NewlineConvertToSpace);
 }
 
 FReply SChangelistTableRow::OnDrop(const FGeometry& InGeometry, const FDragDropEvent& InDragDropEvent)
diff --git a/Engine/Source/Editor/SourceControlWindows/Private/SSourceControlCommon.cpp b/Engine/Source/Editor/SourceControlWindows/Private/SSourceControlCommon.cpp
--- a/Engine/Source/Editor/SourceControlWindows/Private/SSourceControlCommon.cpp
+++ b/Engine/Source/Editor/SourceControlWindows/Private/SSourceControlCommon.cpp
@@ -4,6 +4,7 @@
 
 #include "Algo/Count.h"
 #include "Algo/Find.h"
+#include "Algo/Replace.h"
 #include "AssetRegistry/AssetData.h"
 #include "ActorFolder.h"
 #include "ActorFolderDesc.h"
@@ -416,22 +417,33 @@ FText GetDefaultMultipleAsset()
 	return LOCTEXT("SourceCOntrol_ManyAssetType", "Multiple Assets");
 }
 
-FText GetSingleLineChangelistDescription(const FText& InFullDescription)
+FText GetSingleLineChangelistDescription(const FText& InFullDescription, ESingleLineFlags Flags)
 {
 	FString DescriptionTextAsString = InFullDescription.ToString();
-	DescriptionTextAsString.TrimStartInline();
+	DescriptionTextAsString.TrimStartAndEndInline();
 
-	int32 NewlineStartIndex = INDEX_NONE;
-	DescriptionTextAsString.FindChar(TCHAR('\n'), NewlineStartIndex);
-	if (NewlineStartIndex != INDEX_NONE)
+	if ((Flags & ESingleLineFlags::Mask_NewlineBehavior) == ESingleLineFlags::NewlineConvertToSpace)
 	{
-		DescriptionTextAsString.LeftInline(NewlineStartIndex);
+		static constexpr TCHAR Replacer = TCHAR(' ');
+		// Replace all non-space whitespace characters with space
+		Algo::ReplaceIf(DescriptionTextAsString,
+			[](TCHAR C) { return FChar::IsWhitespace(C) && C != Replacer; },
+			Replacer);
 	}
-
-	// Trim any trailing carriage returns
-	if (DescriptionTextAsString.EndsWith(TEXT("\r"), ESearchCase::CaseSensitive))
+	else
 	{
-		DescriptionTextAsString.LeftChopInline(1);
+		int32 NewlineStartIndex = INDEX_NONE;
+		DescriptionTextAsString.FindChar(TCHAR('\n'), NewlineStartIndex);
+		if (NewlineStartIndex != INDEX_NONE)
+		{
+			DescriptionTextAsString.LeftInline(NewlineStartIndex);
+		}
+
+		// Trim any trailing carriage returns
+		if (DescriptionTextAsString.EndsWith(TEXT("\r"), ESearchCase::CaseSensitive))
+		{
+			DescriptionTextAsString.LeftChopInline(1);
+		}
 	}
 
 	return InFullDescription.IsCultureInvariant() ? FText::AsCultureInvariant(DescriptionTextAsString) : FText::FromString(DescriptionTextAsString);
diff --git a/Engine/Source/Editor/SourceControlWindows/Private/SSourceControlCommon.h b/Engine/Source/Editor/SourceControlWindows/Private/SSourceControlCommon.h
--- a/Engine/Source/Editor/SourceControlWindows/Private/SSourceControlCommon.h
+++ b/Engine/Source/Editor/SourceControlWindows/Private/SSourceControlCommon.h
@@ -383,10 +383,17 @@ namespace SSourceControlCommon
 	FText GetDefaultUnknownAssetType();
 	FText GetDefaultMultipleAsset();
 
+	enum ESingleLineFlags
+	{
+		NewlineTerminates		=0x0,
+		NewlineConvertToSpace	=0x1,
+		Mask_NewlineBehavior	=0x1,
+	};
+	ENUM_CLASS_FLAGS(ESingleLineFlags);
 	/**
 	 * returns the first non-whitespace line, or an empty FText if InFullDescription is empty or only whitespace
 	 */
-	FText GetSingleLineChangelistDescription(const FText& InFullDescription);
+	FText GetSingleLineChangelistDescription(const FText& InFullDescription, ESingleLineFlags Flags = ESingleLineFlags::NewlineTerminates);
 
 	void ExecuteChangelistOperationWithSlowTaskWrapper(const FText& Message, const TFunction<void()>& ChangelistTask);
 	void ExecuteUncontrolledChangelistOperationWithSlowTaskWrapper(const FText& Message, const TFunction<void()>& UncontrolledChangelistTask);
