Fixed crash when reverting Skeletal Mesh unsaved changes using an async task at SSourceControlChangelistRows, while the asset is opened, due to the async task being executed by the preview world of the SkeletalMesh being reverted. Changed the Async task for a CoreTicker call to Execute on game thread.

diff --git a/Engine/Source/Editor/SourceControlWindows/Private/SSourceControlChangelistRows.cpp b/Engine/Source/Editor/SourceControlWindows/Private/SSourceControlChangelistRows.cpp
--- a/Engine/Source/Editor/SourceControlWindows/Private/SSourceControlChangelistRows.cpp
+++ b/Engine/Source/Editor/SourceControlWindows/Private/SSourceControlChangelistRows.cpp
@@ -16,7 +16,7 @@
 #include "Widgets/Input/SButton.h"
 #include "Widgets/Text/STextBlock.h"
 #include "Algo/Transform.h"
-#include "Async/Async.h"
+#include "Containers/Ticker.h"
 #include "RevisionControlStyle/RevisionControlStyle.h"
 #include "Styling/SlateStyleRegistry.h"
 #include "Styling/StarshipCoreStyle.h"
@@ -654,7 +654,7 @@ TSharedRef<SWidget> SOfflineFileTableRow::GenerateWidgetForColumn(const FName& C
 					}
 					
 					DiscardSwitcher->SetActiveWidgetIndex(1);
-					Async(EAsyncExecution::TaskGraphMainThread,
+					ExecuteOnGameThread(UE_SOURCE_LOCATION,
 						[this]
 						{
 							TArray<FString> PackageToReload { TreeItem->GetPackageName().ToString() };
@@ -666,8 +666,7 @@ TSharedRef<SWidget> SOfflineFileTableRow::GenerateWidgetForColumn(const FName& C
 								bAllowReloadWorld,
 								bInteractive
 							);
-						},
-						[] {  });
+						});
 					
 					return FReply::Handled();
 				});
