In SSourceControlChangelistsWidget::OnRefreshSourceControlWidgets handle the detection of a duplicate item gracefully instead of with a check.

This case can occur when a submit is done external to the editor ( e.g. from P4V ) and has also been reported to happen after submits from within the editor.  Now, when it is detected we request a changelists refresh in the next tick by setting bShouldRefresh to true.

This changelist also fixes another issue where files in the list view would get cleared when the selected changelist changed and then never get re-added to the list view.

diff --git a/Engine/Source/Editor/SourceControlWindows/Private/SSourceControlChangelists.cpp b/Engine/Source/Editor/SourceControlWindows/Private/SSourceControlChangelists.cpp
--- a/Engine/Source/Editor/SourceControlWindows/Private/SSourceControlChangelists.cpp
+++ b/Engine/Source/Editor/SourceControlWindows/Private/SSourceControlChangelists.cpp
@@ -995,6 +995,7 @@ void SSourceControlChangelistsWidget::OnRefreshSourceControlWidgets(int64 CurrUp
 	const int64 PrevUpdateNum = CurrUpdateNum - 1;
 	int32 DisplayedItemCount = 0;
 	int32 VisitedItemCount = 0;
+	TArray<IChangelistTreeItem*> DuplicateItems;
 
 	TSharedPtr<IChangelistTreeItem> SelectedChangelistItem = ChangelistTreeView->GetNumItemsSelected() > 0 ? ChangelistTreeView->GetSelectedItems()[0] : TSharedPtr<IChangelistTreeItem>();
 
@@ -1004,11 +1005,20 @@ void SSourceControlChangelistsWidget::OnRefreshSourceControlWidgets(int64 CurrUp
 			(SelectedChangelistItem->GetTreeItemType() == IChangelistTreeItem::Changelist || SelectedChangelistItem->GetTreeItemType() == IChangelistTreeItem::ShelvedChangelist);
 	};
 
-	auto OnSourceControlCachedItemVisited = [CurrUpdateNum, &VisitedItemCount](IChangelistTreeItem* Item)
+	auto OnSourceControlCachedItemVisited = [CurrUpdateNum, &VisitedItemCount, &DuplicateItems](IChangelistTreeItem* Item)
 	{
-		checkf(Item->VisitedUpdateNum < CurrUpdateNum, TEXT("The same revision control item was visited twice. It is likely present in more than one changelist"));
-		Item->VisitedUpdateNum = CurrUpdateNum;
-		++VisitedItemCount;
+		if (Item->VisitedUpdateNum < CurrUpdateNum)
+		{
+			Item->VisitedUpdateNum = CurrUpdateNum;
+			++VisitedItemCount;
+		}
+		else
+		{
+			// The same revision control item was visited twice. It is likely present in more than one changelist
+			// This can happen in the case of a submit from outside of the editor ( e.g. from P4V )
+			// Handle it gracefully by ignoring the duplicate and then request a full changelists refresh for next tick
+			DuplicateItems.Emplace(Item);
+		}
 	};
 
 	auto OnNewSourceControlItemCreated = [this, &OnSourceControlCachedItemVisited](TSharedPtr<void> Key, TSharedPtr<IChangelistTreeItem> Item)
@@ -1115,7 +1125,7 @@ void SSourceControlChangelistsWidget::OnRefreshSourceControlWidgets(int64 CurrUp
 				OnNewSourceControlItemCreated(FileState, FileItem);
 			}
 
-			if (bChangelistPassedFilter && FileTextFilter->PassesFilter(*FileItem))
+			if (bChangelistPassedFilter && FileTextFilter->PassesFilter(*FileItem) && !DuplicateItems.Contains(FileItem.Get()))
 			{
 				ChangelistTreeItem->AddChild(FileItem.ToSharedRef());
 				if (bChangelistSelected)
@@ -1155,6 +1165,9 @@ void SSourceControlChangelistsWidget::OnRefreshSourceControlWidgets(int64 CurrUp
 
 	RemoveDiscardedSourceControlItems();
 	ChangelistTreeView->RequestTreeRefresh();
+
+	// If a duplicate item was detected then request a changelists refresh in the next tick 
+	bShouldRefresh |= !DuplicateItems.IsEmpty();
 }
 
 void SSourceControlChangelistsWidget::OnRefreshUncontrolledChangelistWidgets(int64 CurrUpdateNum, const TFunction<void(TSharedPtr<IFileViewTreeItem>&)>& AddItemToFileView)
@@ -1351,7 +1364,8 @@ void SSourceControlChangelistsWidget::OnRefreshUI(ERefreshFlags RefreshFlag)
 	auto AddItemToFileView = [this, &NewDisplayedItemCount, &bDisplayedIconsPriorityChanged, PrevUpdateNum, CurrUpdateNum](TSharedPtr<IFileViewTreeItem>& Item)
 	{
 		// If the items wasn't displayed last update.
-		if (Item->DisplayedUpdateNum != PrevUpdateNum)
+		// Or, if FileListNodes was cleared from OnChangelistSelectionChanged
+		if (Item->DisplayedUpdateNum != PrevUpdateNum || !FileListNodes.Contains(Item))
 		{
 			checkfSlow(!FileListNodes.Contains(Item), TEXT("Inserting duplicated items. Something is wrong with the display update number."))
 
@@ -3531,6 +3545,8 @@ void SSourceControlChangelistsWidget::OnChangelistSelectionChanged(TSharedPtr<IC
 	TRACE_CPUPROFILER_EVENT_SCOPE(SSourceControlChangelistsWidget::OnChangelistSelectionChanged);
 
 	FileListNodes.Reset();
+	GetActiveFileListView().RequestListRefresh();
+
 	bUpdateMonitoredFileStatusList = true;
 
 	if (SelectedChangelist) // Can be a Changelist, Uncontrolled Changelist or Shelved Changelist
