Source Control Changelist View
* Fix initial refresh UI not being call for all changelist type (esp. the Uncontrolled Cl)

diff --git a/Engine/Source/Editor/SourceControlWindows/Private/SSourceControlChangelists.cpp b/Engine/Source/Editor/SourceControlWindows/Private/SSourceControlChangelists.cpp
--- a/Engine/Source/Editor/SourceControlWindows/Private/SSourceControlChangelists.cpp
+++ b/Engine/Source/Editor/SourceControlWindows/Private/SSourceControlChangelists.cpp
@@ -662,7 +662,7 @@ void SSourceControlChangelistsWidget::Tick(const FGeometry& AllottedGeometry, co
 
 	if (bShouldRefresh)
 	{
-		if (bInitialRefreshDone)
+		if (!bInitialRefreshDone)
 		{
 			// Ensure the UI is built from the current cached states once because all UI updates are hooked on 'state change' callbacks and those will only
 			// be called when the state changed.
