SkeinUE: Generalized 'checked out by' dialog by moving it to /SourceControlWindows. Folluw up CL from #33935263.

diff --git a/Engine/Source/Editor/SourceControlWindows/Private/SSourceControlCheckedOutDialog.cpp b/Engine/Source/Editor/SourceControlWindows/Private/SSourceControlCheckedOutDialog.cpp
new file mode 100644
--- /dev/null
+++ b/Engine/Source/Editor/SourceControlWindows/Private/SSourceControlCheckedOutDialog.cpp
@@ -0,0 +1,316 @@
+// Copyright Epic Games, Inc. All Rights Reserved.
+
+#include "SSourceControlCheckedOutDialog.h"
+#include "Brushes/SlateDynamicImageBrush.h"
+#include "Styling/AppStyle.h"
+#include "Widgets/SWindow.h"
+#include "Widgets/Layout/SUniformGridPanel.h"
+#include "Widgets/Input/SButton.h"
+#include "Widgets/Input/SHyperlink.h"
+#include "Widgets/Text/SRichTextBlock.h"
+#include "Framework/Application/SlateApplication.h"
+#include "Misc/Paths.h"
+#include "HAL/PlatformFileManager.h"
+#include "Framework/Application/SlateApplication.h"
+#include "SPrimaryButton.h"
+#include "ISourceControlModule.h"
+#include "SourceControlAssetDataCache.h"
+#include "Algo/Find.h"
+#include "Algo/Count.h"
+#include "ActorFolder.h"
+#include "ActorFolderDesc.h"
+
+#define LOCTEXT_NAMESPACE "SSourceControlConflict"
+
+namespace SSourceControlConflictWarningWidgetDefs
+{
+	const FName ColumnID_AssetLabel("Asset");
+	const FName ColumnID_AssetClassLabel("AssetClass");
+	const FName ColumnID_UserNameLabel("UserName");
+}
+
+void SSourceControlCheckedOutDialog::Construct(const FArguments& InArgs)
+{
+	ParentFrame = InArgs._ParentWindow;
+	SortByColumn = SSourceControlConflictWarningWidgetDefs::ColumnID_AssetLabel;
+	SortMode = EColumnSortMode::Ascending;
+
+	ListViewItems.Reserve(InArgs._Items.Num());
+	for (const auto & Item : InArgs._Items)
+	{
+		ListViewItems.Add(MakeShared<FFileTreeItem>(Item));
+	}
+
+	TSharedPtr<SHeaderRow> HeaderRowWidget;
+	HeaderRowWidget = SNew(SHeaderRow);
+
+	if (InArgs._ShowColumnAssetName)
+	{
+		HeaderRowWidget->AddColumn(
+			SHeaderRow::Column(SSourceControlConflictWarningWidgetDefs::ColumnID_AssetLabel)
+			.DefaultLabel(LOCTEXT("AssetColumnLabel", "Asset Name"))
+			.SortMode(this, &SSourceControlCheckedOutDialog::GetColumnSortMode, SSourceControlConflictWarningWidgetDefs::ColumnID_AssetLabel)
+			.OnSort(this, &SSourceControlCheckedOutDialog::OnColumnSortModeChanged)
+			.FillWidth(0.5f)
+		);
+	}
+
+	if (InArgs._ShowColumnAssetClass)
+	{
+		HeaderRowWidget->AddColumn(
+			SHeaderRow::Column(SSourceControlConflictWarningWidgetDefs::ColumnID_AssetClassLabel)
+			.DefaultLabel(LOCTEXT("AssetClassLabel", "Asset Class"))
+			.FillWidth(0.5f)
+		);
+	}
+
+	if (InArgs._ShowColumnUserName)
+	{
+		HeaderRowWidget->AddColumn(
+			SHeaderRow::Column(SSourceControlConflictWarningWidgetDefs::ColumnID_UserNameLabel)
+			.DefaultLabel(LOCTEXT("UserNameLabel", "User Name"))
+			.FillWidth(0.25f)
+		);
+	}
+
+	TSharedPtr<SHorizontalBox> ButtonsBox = SNew(SHorizontalBox);
+
+	bool bShowCheckBox = !InArgs._CheckBoxText.IsEmpty();
+	if (bShowCheckBox)
+	{
+		CheckBox = SNew(SCheckBox);
+		CheckBox->SetIsChecked(ECheckBoxState::Checked);
+
+		ButtonsBox->AddSlot()
+		[
+			SNew(SHorizontalBox)
+			+ SHorizontalBox::Slot()
+			.HAlign(HAlign_Left)
+			[
+				SNew(SHorizontalBox)
+				+ SHorizontalBox::Slot()
+				.AutoWidth()
+				.Padding(5, 0)
+				[
+					CheckBox.ToSharedRef()
+				]
+				+ SHorizontalBox::Slot()
+				.AutoWidth()
+				.Padding(5, 5)
+				[
+					SNew(STextBlock)
+					.Text(InArgs._CheckBoxText)
+				]
+			]
+			+ SHorizontalBox::Slot()
+			.HAlign(HAlign_Right)
+			.Padding(5, 0)
+			[
+				SAssignNew(CloseButton, SPrimaryButton)
+				.Text(InArgs._CloseText)
+				.OnClicked(this, &SSourceControlCheckedOutDialog::CloseClicked)
+			]
+		];
+	}
+	else
+	{
+		ButtonsBox->AddSlot()
+		.AutoWidth()
+		.Padding(5, 0)
+		[
+			SAssignNew(CloseButton, SPrimaryButton)
+			.Text(InArgs._CloseText)
+			.OnClicked(this, &SSourceControlCheckedOutDialog::CloseClicked)
+		];
+	}
+
+	TSharedPtr<SVerticalBox> Contents;
+
+	ChildSlot
+	[
+		SNew(SBorder)
+		.BorderImage(FAppStyle::GetBrush("ToolPanel.GroupBorder"))
+		.Padding(16.f)
+		[
+			SAssignNew(Contents, SVerticalBox)
+		]
+	];
+
+	Contents->AddSlot()
+	.AutoHeight()
+	[
+		SNew(SHorizontalBox)
+		+ SHorizontalBox::Slot()
+		.AutoWidth()
+		.VAlign(VAlign_Top)
+		.HAlign(HAlign_Left)
+		.Padding(FMargin(0.0f, 0.0f, 16.0f, 0.0f))
+		[
+			SNew(SImage)
+			.Image(FAppStyle::Get().GetBrush("Icons.WarningWithColor.Large"))
+		]
+		+SHorizontalBox::Slot()
+		.VAlign(VAlign_Center)
+		[
+			SNew(SRichTextBlock)
+			.DecoratorStyleSet(&FAppStyle::Get())
+			.Text(InArgs._MessageText)
+			.AutoWrapText(true)
+		]
+	];
+
+	Contents->AddSlot()
+	.Padding(0.0f, 16.0f, 0.0f, 0.0f)
+	[
+		SNew(SBorder)
+		.Visibility_Lambda([this]() -> EVisibility
+		{
+			return ListViewItems.Num() > 0 ? EVisibility::Visible : EVisibility::Collapsed;
+		})
+		[
+			SNew(SBox)
+			.HeightOverride(200.0f)
+			.WidthOverride(800.0f)
+			[
+				SAssignNew(ListView, SListView<FFileTreeItemRef>)
+				.ListItemsSource(&ListViewItems)
+				.OnGenerateRow(this, &SSourceControlCheckedOutDialog::OnGenerateRowForList)
+				.HeaderRow(HeaderRowWidget)
+				.SelectionMode(ESelectionMode::Single)
+			]
+		]
+	];
+
+	Contents->AddSlot()
+	.AutoHeight()
+	.Padding(0.0f, 16.0f, 0.0f, 0.0f)
+	.HAlign(bShowCheckBox ? HAlign_Fill : HAlign_Right)
+	.VAlign(VAlign_Bottom)
+	[
+		ButtonsBox.ToSharedRef()
+	];
+
+	RequestSort();
+}
+
+TSharedRef<ITableRow> SSourceControlCheckedOutDialog::OnGenerateRowForList(FFileTreeItemRef InItem, const TSharedRef<STableViewBase>& OwnerTable)
+{
+	TSharedRef<ITableRow> Row =
+		SNew(SSourceControlCheckedOutDialogListRow, OwnerTable)
+		.Item(InItem);
+
+	return Row;
+}
+
+EColumnSortMode::Type SSourceControlCheckedOutDialog::GetColumnSortMode(const FName ColumnId) const
+{
+	if (SortByColumn != ColumnId)
+	{
+		return EColumnSortMode::None;
+	}
+
+	return SortMode;
+}
+
+void SSourceControlCheckedOutDialog::OnColumnSortModeChanged(const EColumnSortPriority::Type SortPriority, const FName& ColumnId, const EColumnSortMode::Type InSortMode)
+{
+	SortByColumn = ColumnId;
+	SortMode = InSortMode;
+
+	RequestSort();
+}
+
+void SSourceControlCheckedOutDialog::RequestSort()
+{
+	SortTree();
+
+	ListView->RequestListRefresh();
+}
+
+void SSourceControlCheckedOutDialog::SortTree()
+{
+	if (SortByColumn == SSourceControlConflictWarningWidgetDefs::ColumnID_AssetLabel)
+	{
+		if (SortMode == EColumnSortMode::Ascending)
+		{
+			ListViewItems.Sort([this](const FFileTreeItemRef& A, const FFileTreeItemRef& B) {
+				return A->GetAssetName().CompareTo(B->GetAssetName()) < 0; });
+		}
+		else if (SortMode == EColumnSortMode::Descending)
+		{
+			ListViewItems.Sort([this](const FFileTreeItemRef& A, const FFileTreeItemRef& B) {
+				return A->GetAssetName().CompareTo(B->GetAssetName()) >= 0; });
+		}
+	}
+}
+
+FReply SSourceControlCheckedOutDialog::CloseClicked()
+{
+	ParentFrame.Pin()->RequestDestroyWindow();
+
+	return FReply::Handled();
+}
+
+FReply SSourceControlCheckedOutDialog::OnKeyDown(const FGeometry& MyGeometry, const FKeyEvent& InKeyEvent)
+{
+	if ((InKeyEvent.GetKey() == EKeys::Escape && CloseButton.IsValid()))
+	{
+		return CloseClicked();
+	}
+
+	return FReply::Unhandled();
+}
+
+void SSourceControlCheckedOutDialogListRow::Construct(const FArguments& InArgs, const TSharedRef<STableViewBase>& InOwnerTableView)
+{
+	Item = InArgs._Item;
+
+	SMultiColumnTableRow<FFileTreeItemRef>::Construct(FSuperRowType::FArguments(), InOwnerTableView);
+}
+
+TSharedRef<SWidget> SSourceControlCheckedOutDialogListRow::GenerateWidgetForColumn(const FName& ColumnName)
+{
+	const FMargin RowPadding(8, 2, 2, 2);
+
+	TSharedPtr<SWidget> ItemContentWidget;
+	if (ColumnName == SSourceControlConflictWarningWidgetDefs::ColumnID_AssetLabel)
+	{
+		ItemContentWidget = SNew(SHorizontalBox)
+			+ SHorizontalBox::Slot()
+			.Padding(RowPadding)
+			.VAlign(VAlign_Center)
+			[
+				SNew(STextBlock)
+				.Text(Item->GetAssetName())
+				.ToolTipText(Item->GetPackageName())
+			];
+	}
+	else if (ColumnName == SSourceControlConflictWarningWidgetDefs::ColumnID_AssetClassLabel)
+	{
+		ItemContentWidget = SNew(SHorizontalBox)
+			+ SHorizontalBox::Slot()
+			.Padding(RowPadding)
+			.VAlign(VAlign_Center)
+			[
+				SNew(STextBlock)
+				.Text(Item->GetAssetType())
+				.ToolTipText(Item->GetAssetType())
+			];
+	}
+	else if (ColumnName == SSourceControlConflictWarningWidgetDefs::ColumnID_UserNameLabel)
+	{
+		ItemContentWidget = SNew(SHorizontalBox)
+			+ SHorizontalBox::Slot()
+			.Padding(RowPadding)
+			.VAlign(VAlign_Center)
+			[
+				SNew(STextBlock)
+				.Text(Item->GetCheckedOutByUser())
+				.ToolTipText(Item->GetCheckedOutByUser())
+			];
+	}
+
+	return ItemContentWidget.ToSharedRef();
+}
+
+#undef LOCTEXT_NAMESPACE
diff --git a/Engine/Source/Editor/SourceControlWindows/Private/SSourceControlCheckedOutDialog.h b/Engine/Source/Editor/SourceControlWindows/Private/SSourceControlCheckedOutDialog.h
new file mode 100644
--- /dev/null
+++ b/Engine/Source/Editor/SourceControlWindows/Private/SSourceControlCheckedOutDialog.h
@@ -0,0 +1,135 @@
+// Copyright Epic Games, Inc. All Rights Reserved.
+
+#pragma once
+
+#include "CoreMinimal.h"
+#include "Layout/Visibility.h"
+#include "Input/Reply.h"
+#include "Widgets/DeclarativeSyntaxSupport.h"
+#include "Widgets/SWidget.h"
+#include "Widgets/SCompoundWidget.h"
+#include "Widgets/Views/SListView.h"
+#include "Widgets/Views/STableRow.h"
+#include "Widgets/Input/SCheckBox.h"
+#include "SSourceControlCommon.h"
+
+class SWindow;
+
+class SSourceControlCheckedOutDialog : public SCompoundWidget
+{
+public:
+	SLATE_BEGIN_ARGS(SSourceControlCheckedOutDialog)
+		: _ParentWindow()
+		, _Items()
+		, _ShowColumnAssetName(true)
+		, _ShowColumnAssetClass(true)
+		, _ShowColumnUserName(true)
+		, _MessageText()
+		, _CloseText()
+		, _CheckBoxText()
+		{ }
+
+		SLATE_ARGUMENT(TSharedPtr<SWindow>, ParentWindow)
+		SLATE_ARGUMENT(TArray<FSourceControlStateRef>, Items)
+		SLATE_ARGUMENT(bool, ShowColumnAssetName)
+		SLATE_ARGUMENT(bool, ShowColumnAssetClass)
+		SLATE_ARGUMENT(bool, ShowColumnUserName)
+		SLATE_ARGUMENT(FText, MessageText)
+		SLATE_ARGUMENT(FText, CloseText)
+		SLATE_ARGUMENT(FText, CheckBoxText)
+
+	SLATE_END_ARGS()
+
+	/** Constructs the widget */
+	void Construct(const FArguments& InArgs);
+
+	/** Get whether to hide the checkbox is checked */
+	bool IsCheckBoxChecked() { return CheckBox.IsValid() ? CheckBox->IsChecked() : false; }
+
+private:
+	/** Callback to generate ListBoxRows */
+	TSharedRef<ITableRow> OnGenerateRowForList(FFileTreeItemRef InItem, const TSharedRef<STableViewBase>& OwnerTable);
+
+	/**
+	 * Returns the current column sort mode (ascending or descending) if the ColumnId parameter matches the current
+	 * column to be sorted by, otherwise returns EColumnSortMode_None.
+	 *
+	 * @param	ColumnId	Column ID to query sort mode for.
+	 * @return	The sort mode for the column, or EColumnSortMode_None if it is not known.
+	 */
+	EColumnSortMode::Type GetColumnSortMode(const FName ColumnId) const;
+
+	/**
+	 * Callback for SHeaderRow::Column::OnSort, called when the column to sort by is changed.
+	 *
+	 * @param	ColumnId	The new column to sort by
+	 * @param	InSortMode	The sort mode (ascending or descending)
+	 */
+	void OnColumnSortModeChanged(const EColumnSortPriority::Type SortPriority, const FName& ColumnId, const EColumnSortMode::Type InSortMode);
+
+	/**
+	 * Requests that the source list data be sorted according to the current sort column and mode,
+	 * and refreshes the list view.
+	 */
+	void RequestSort();
+
+	/**
+	 * Sorts the source list data according to the current sort column and mode.
+	 */
+	void SortTree();
+
+	/**
+	 * Refreshes the assetdata for the items from the AssetDataCache
+	 */
+	void RefreshItemAssetData();
+
+	/** Called when the cancel button is clicked */
+	FReply CloseClicked();
+
+private:
+	//~ Begin SWidget Interface
+	virtual FReply OnKeyDown(const FGeometry& MyGeometry, const FKeyEvent& InKeyEvent) override;
+	//~ End SWidget Interface
+
+	/** Pointer to the parent modal window */
+	TWeakPtr<SWindow> ParentFrame;
+
+	/** Collection of objects to display in the List View */
+	TArray<FFileTreeItemRef> ListViewItems;
+
+	/** ListBox for displaying items */
+	TSharedPtr<SListView<FFileTreeItemRef>> ListView;
+
+	/** Specify which column to sort with */
+	FName SortByColumn;
+
+	/** Currently selected sorting mode */
+	EColumnSortMode::Type SortMode;
+
+	/** The close button widget. */
+	TSharedPtr<SButton> CloseButton;
+
+	/** The checkbox widget. */
+	TSharedPtr<SCheckBox> CheckBox;
+};
+
+class SSourceControlCheckedOutDialogListRow : public SMultiColumnTableRow<FFileTreeItemRef>
+{
+public:
+	SLATE_BEGIN_ARGS(SSourceControlCheckedOutDialogListRow) {}
+
+		SLATE_ARGUMENT(FFileTreeItemPtr, Item)
+
+	SLATE_END_ARGS()
+
+	/** Construct function for this widget */
+	void Construct(const FArguments& InArgs, const TSharedRef<STableViewBase>& InOwnerTableView);
+
+	/** Overridden from SMultiColumnTableRow.  Generates a widget for this column of the list row. */
+	virtual TSharedRef<SWidget> GenerateWidgetForColumn(const FName& ColumnName) override;
+
+private:
+
+	/** The item associated with this row */
+	FFileTreeItemPtr Item;
+};
\ No newline at end of file
diff --git a/Engine/Source/Editor/SourceControlWindows/Private/SourceControlWindows.cpp b/Engine/Source/Editor/SourceControlWindows/Private/SourceControlWindows.cpp
--- a/Engine/Source/Editor/SourceControlWindows/Private/SourceControlWindows.cpp
+++ b/Engine/Source/Editor/SourceControlWindows/Private/SourceControlWindows.cpp
@@ -2,6 +2,7 @@
 
 #include "SourceControlWindows.h"
 #include "SSourceControlSubmit.h"
+#include "SSourceControlCheckedOutDialog.h"
 #include "AssetViewUtils.h"
 #include "FileHelpers.h"
 #include "ISourceControlModule.h"
@@ -541,6 +542,48 @@ bool FSourceControlWindows::PromptForCheckin(bool bUseSourceControlStateCache, c
 	return PromptForCheckin(ResultInfo, InPackageNames, InPendingDeletePaths, InConfigFiles, bUseSourceControlStateCache);
 }
 
+bool FSourceControlWindows::PromptForCheckedOut(bool bUseSourceControlStateCache, TArray<FString>& InFileNames, FCheckedOutSetupInfo& InSetupInfo)
+{
+	TArray<FSourceControlStateRef> Items;
+
+	if (InFileNames.Num() > 0)
+	{
+		ISourceControlProvider& SourceControlProvider = ISourceControlModule::Get().GetProvider();
+
+		if (!bUseSourceControlStateCache)
+		{
+			SourceControlProvider.Execute(ISourceControlOperation::Create<FUpdateStatus>(), InFileNames);
+		}
+
+		SourceControlProvider.GetState(InFileNames, Items, EStateCacheUsage::Use);
+	}
+
+	TSharedRef<SWindow> Window = SNew(SWindow)
+		.SizingRule(ESizingRule::Autosized)
+		.ClientSize(FVector2D(800, 600))
+		.SupportsMaximize(false)
+		.SupportsMinimize(false);
+
+	TSharedRef<SSourceControlCheckedOutDialog> Widget =
+		SNew(SSourceControlCheckedOutDialog)
+		.ParentWindow(Window)
+		.Items(Items)
+		.ShowColumnAssetName(InSetupInfo.bShowColumnAssetName)
+		.ShowColumnAssetClass(InSetupInfo.bShowColumnAssetClass)
+		.ShowColumnUserName(InSetupInfo.bShowColumnUserName)
+		.MessageText(InSetupInfo.MessageText)
+		.CloseText(InSetupInfo.CloseText)
+		.CheckBoxText(InSetupInfo.CheckboxText);
+
+	Window->SetTitle(InSetupInfo.TitleText);
+	Window->SetContent(
+		Widget
+	);
+
+	FSlateApplication::Get().AddModalWindow(Window, NULL);
+
+	return Widget->IsCheckBoxChecked();
+}
 
 // Note that:
 // - FSourceControlWindows::DisplayRevisionHistory() is defined in SSourceControlHistory.cpp
diff --git a/Engine/Source/Editor/SourceControlWindows/Public/SourceControlWindows.h b/Engine/Source/Editor/SourceControlWindows/Public/SourceControlWindows.h
--- a/Engine/Source/Editor/SourceControlWindows/Public/SourceControlWindows.h
+++ b/Engine/Source/Editor/SourceControlWindows/Public/SourceControlWindows.h
@@ -121,6 +121,28 @@ public:
 	 */
 	static bool PromptForCheckin(bool bUseSourceControlStateCache, const TArray<FString>& InPackageNames, const TArray<FString>& InPendingDeletePaths = TArray<FString>(), const TArray<FString>& InConfigFiles = TArray<FString>());
 
+	/**
+	 * Display checked out dialog for the specified filenames
+	 * @param	bUseSourceControlStateCache		Whether to use the cached source control status, or force the status to be updated
+	 * @param	InFileNames						Filenames of files that are checked out
+	 * @param	InSetupInfo						Config options
+	 * 
+	 * @return	true - if the dialog was closed with the optional checkbox checked
+	 */
+	struct FCheckedOutSetupInfo
+	{
+		FText TitleText;
+		FText MessageText;
+		FText CloseText;
+		FText CheckboxText;
+
+		bool bShowColumnAssetName = true;
+		bool bShowColumnAssetClass = true;
+		bool bShowColumnUserName = true;
+	};
+
+	static bool PromptForCheckedOut(bool bUseSourceControlStateCache, TArray<FString>& InFileNames, FCheckedOutSetupInfo& InSetupInfo);
+
 	/**
 	 * Display file revision history for the provided packages
 	 *
@@ -171,8 +193,6 @@ protected:
 	/** Delegate called when the user has decided to cancel the check in process */
 	static void ChoosePackagesToCheckInCancelled(FSourceControlOperationRef InOperation);
 
-	/** Delegate called when the user clicks submit.  if the return value is true the submit dialog is closed.  Otherwise it is left open for further corrections */
-	static bool OnSubmitClicked(TSharedRef<class SSourceControlSubmitWidget> SourceControlWidget);
 private:
 	/** The notification in place while we choose packages to check in */
 	static TWeakPtr<class SNotificationItem> ChoosePackagesToCheckInNotification;
