[Revision Control] Auto select CL description text when moving files between CLs

diff --git a/Engine/Source/Editor/SourceControlWindows/Private/SSourceControlDescription.cpp b/Engine/Source/Editor/SourceControlWindows/Private/SSourceControlDescription.cpp
--- a/Engine/Source/Editor/SourceControlWindows/Private/SSourceControlDescription.cpp
+++ b/Engine/Source/Editor/SourceControlWindows/Private/SSourceControlDescription.cpp
@@ -42,7 +42,6 @@ void SSourceControlDescriptionWidget::Construct(const FArguments& InArgs)
 	const bool bHasItems = (Items && Items->Num() > 0);
 	const bool bShowSelectionDropDown = bHasItems;
 	const bool bDescriptionCanBeEdited = (!bHasItems || (*Items)[0].bCanEditDescription);
-	const bool bSelectTextWhenFocused = !bHasItems;
 
 	CurrentlySelectedItemIndex = 0;
 
@@ -84,7 +83,7 @@ void SSourceControlDescriptionWidget::Construct(const FArguments& InArgs)
 			.Padding(FMargin(16, 0, 16, 16))
 			[
 				SAssignNew(TextBox, SMultiLineEditableTextBox)
-				.SelectAllTextWhenFocused(bSelectTextWhenFocused)
+				.SelectAllTextWhenFocused(bDescriptionCanBeEdited)
 				.IsReadOnly(!bDescriptionCanBeEdited)
 				.AutoWrapText(true)
 				.Text(InArgs._Text)
@@ -215,6 +214,10 @@ TSharedRef<SWidget> SSourceControlDescriptionWidget::GetSelectionContent()
 				FExecuteAction::CreateLambda([this, ItemIndex, &Item]() {
 					TextBox->SetIsReadOnly(!Item.bCanEditDescription);
 					TextBox->SetText(Item.Description);
+					if (Item.bCanEditDescription)
+					{
+						TextBox->SelectAllText();
+					}
 					
 					CurrentlySelectedItemIndex = ItemIndex;
 					})),
