SkeinUE: Improved asset class display name in Snapshot History and Checked Out dialogs.

diff --git a/Engine/Source/Editor/SourceControlWindows/Private/SSourceControlCheckedOutDialog.cpp b/Engine/Source/Editor/SourceControlWindows/Private/SSourceControlCheckedOutDialog.cpp
--- a/Engine/Source/Editor/SourceControlWindows/Private/SSourceControlCheckedOutDialog.cpp
+++ b/Engine/Source/Editor/SourceControlWindows/Private/SSourceControlCheckedOutDialog.cpp
@@ -19,6 +19,7 @@
 #include "Algo/Count.h"
 #include "ActorFolder.h"
 #include "ActorFolderDesc.h"
+#include "AssetDefinitionRegistry.h"
 
 #define LOCTEXT_NAMESPACE "SSourceControlConflict"
 
@@ -44,7 +45,25 @@ void SSourceControlCheckedOutDialog::Construct(const FArguments& InArgs)
 	TSharedPtr<SHeaderRow> HeaderRowWidget;
 	HeaderRowWidget = SNew(SHeaderRow);
 
-	if (InArgs._ShowColumnAssetName)
+	bool bShowColumnAssetName = InArgs._ShowColumnAssetName;
+	bool bShowColumnAssetClass = InArgs._ShowColumnAssetClass;
+	bool bShowColumnUserName = InArgs._ShowColumnUserName;
+
+	if (bShowColumnUserName)
+	{
+		bool bAnyCheckedOut = false;
+		for (auto& Item : ListViewItems)
+		{
+			if (!Item->GetCheckedOutByUser().IsEmpty())
+			{
+				bAnyCheckedOut = true;
+				break;
+			}
+		}
+		bShowColumnUserName = bAnyCheckedOut;
+	}
+
+	if (bShowColumnAssetName)
 	{
 		HeaderRowWidget->AddColumn(
 			SHeaderRow::Column(SSourceControlConflictWarningWidgetDefs::ColumnID_AssetLabel)
@@ -55,7 +74,7 @@ void SSourceControlCheckedOutDialog::Construct(const FArguments& InArgs)
 		);
 	}
 
-	if (InArgs._ShowColumnAssetClass)
+	if (bShowColumnAssetName)
 	{
 		HeaderRowWidget->AddColumn(
 			SHeaderRow::Column(SSourceControlConflictWarningWidgetDefs::ColumnID_AssetClassLabel)
@@ -64,7 +83,7 @@ void SSourceControlCheckedOutDialog::Construct(const FArguments& InArgs)
 		);
 	}
 
-	if (InArgs._ShowColumnUserName)
+	if (bShowColumnUserName)
 	{
 		HeaderRowWidget->AddColumn(
 			SHeaderRow::Column(SSourceControlConflictWarningWidgetDefs::ColumnID_UserNameLabel)
@@ -293,7 +312,7 @@ TSharedRef<SWidget> SSourceControlCheckedOutDialogListRow::GenerateWidgetForColu
 			.VAlign(VAlign_Center)
 			[
 				SNew(STextBlock)
-				.Text(Item->GetAssetType())
+				.Text(Item->GetAssetTypeName())
 				.ToolTipText(Item->GetAssetType())
 			];
 	}
diff --git a/Engine/Source/Editor/SourceControlWindows/Private/SSourceControlCommon.cpp b/Engine/Source/Editor/SourceControlWindows/Private/SSourceControlCommon.cpp
--- a/Engine/Source/Editor/SourceControlWindows/Private/SSourceControlCommon.cpp
+++ b/Engine/Source/Editor/SourceControlWindows/Private/SSourceControlCommon.cpp
@@ -8,6 +8,7 @@
 #include "AssetRegistry/AssetData.h"
 #include "ActorFolder.h"
 #include "ActorFolderDesc.h"
+#include "AssetDefinitionRegistry.h"
 #include "AssetToolsModule.h"
 #include "Styling/AppStyle.h"
 #include "ISourceControlModule.h"
@@ -105,13 +106,28 @@ static FString RetrieveAssetPath(const FAssetData& InAssetData)
 	return Path;
 }
 
-static void RefreshAssetInformationInternal(const TArray<FAssetData>& Assets, const FString& InFilename, FString& OutAssetName, FString& OutAssetPath, FString& OutAssetType, FText& OutPackageName, FColor& OutAssetTypeColor)
+static FString RetrieveAssetTypeName(const FAssetData& InAssetData)
+{
+	if (UAssetDefinitionRegistry* AssetDefinitionRegistry = UAssetDefinitionRegistry::Get())
+	{
+		const UAssetDefinition* AssetDefinition = AssetDefinitionRegistry->GetAssetDefinitionForAsset(InAssetData);
+		if (AssetDefinition)
+		{
+			return AssetDefinition->GetAssetDisplayName().ToString();
+		}
+	}
+
+	return InAssetData.AssetClassPath.ToString();
+}
+
+static void RefreshAssetInformationInternal(const TArray<FAssetData>& Assets, const FString& InFilename, FString& OutAssetName, FString& OutAssetPath, FString& OutAssetType, FString& OutAssetTypeName, FText& OutPackageName, FColor& OutAssetTypeColor)
 {
 	// Initialize display-related members
 	FString Filename = InFilename;
 	FString TempAssetName = SSourceControlCommon::GetDefaultAssetName().ToString();
 	FString TempAssetPath = Filename;
 	FString TempAssetType = SSourceControlCommon::GetDefaultAssetType().ToString();
+	FString TempAssetTypeName = SSourceControlCommon::GetDefaultAssetType().ToString();
 	FString TempPackageName = Filename;
 	FColor TempAssetColor = FColor(		// Copied from ContentBrowserCLR.cpp
 		127 + FColor::Red.R / 2,	// Desaturate the colors a bit (GB colors were too.. much)
@@ -131,6 +147,7 @@ static void RefreshAssetInformationInternal(const TArray<FAssetData>& Assets, co
 			TempAssetName = RetrieveAssetName(AssetData);
 			TempAssetPath = RetrieveAssetPath(AssetData);
 			TempAssetType = AssetData.AssetClassPath.ToString();
+			TempAssetTypeName = RetrieveAssetTypeName(AssetData);
 
 			const FAssetToolsModule& AssetToolsModule = FModuleManager::LoadModuleChecked<FAssetToolsModule>(TEXT("AssetTools"));
 			const TSharedPtr<IAssetTypeActions> AssetTypeActions = AssetToolsModule.Get().GetAssetTypeActionsForClass(AssetData.GetClass()).Pin();
@@ -148,6 +165,7 @@ static void RefreshAssetInformationInternal(const TArray<FAssetData>& Assets, co
 		{
 			TempAssetName = RetrieveAssetName(Assets[0]);
 			TempAssetPath = RetrieveAssetPath(Assets[0]);
+			TempAssetTypeName = RetrieveAssetTypeName(Assets[0]);
 
 			for (int32 i = 1; i < Assets.Num(); ++i)
 			{
@@ -178,12 +196,14 @@ static void RefreshAssetInformationInternal(const TArray<FAssetData>& Assets, co
 		TempAssetName = FPaths::GetCleanFilename(Filename);
 		TempPackageName = Filename; // put back original package name if the try failed
 		TempAssetType = FText::Format(SSourceControlCommon::GetDefaultUnknownAssetType(), FText::FromString(FPaths::GetExtension(Filename).ToUpper())).ToString();
+		TempAssetTypeName = TempAssetType;
 	}
 
 	// Finally, assign the temp variables to the member variables
 	OutAssetName = TempAssetName;
 	OutAssetPath = TempAssetPath;
 	OutAssetType = TempAssetType;
+	OutAssetTypeName = TempAssetTypeName;
 	OutAssetTypeColor = TempAssetColor;
 	OutPackageName = FText::FromString(TempPackageName);
 }
@@ -269,10 +289,11 @@ void FFileTreeItem::RefreshAssetInformation()
 {
 	// Initialize display-related members
 	static TArray<FAssetData> NoAssets;
-	RefreshAssetInformationInternal(Assets.IsValid() ? *Assets : NoAssets, FileState->GetFilename(), AssetNameStr, AssetPathStr, AssetTypeStr, PackageName, AssetTypeColor);
+	RefreshAssetInformationInternal(Assets.IsValid() ? *Assets : NoAssets, FileState->GetFilename(), AssetNameStr, AssetPathStr, AssetTypeStr, AssetTypeNameStr, PackageName, AssetTypeColor);
 	AssetName = FText::FromString(AssetNameStr);
 	AssetPath = FText::FromString(AssetPathStr);
 	AssetType = FText::FromString(AssetTypeStr);
+	AssetTypeName = FText::FromString(AssetTypeNameStr);
 }
 
 FText FFileTreeItem::GetAssetName() const
@@ -327,10 +348,11 @@ FOfflineFileTreeItem::FOfflineFileTreeItem(const FString& InFilename)
 
 void FOfflineFileTreeItem::RefreshAssetInformation()
 {
-	RefreshAssetInformationInternal(Assets, Filename, AssetNameStr, AssetPathStr, AssetTypeStr, PackageName, AssetTypeColor);
+	RefreshAssetInformationInternal(Assets, Filename, AssetNameStr, AssetPathStr, AssetTypeStr, AssetTypeNameStr, PackageName, AssetTypeColor);
 	AssetName = FText::FromString(AssetNameStr);
 	AssetPath = FText::FromString(AssetPathStr);
 	AssetType = FText::FromString(AssetTypeStr);
+	AssetTypeName = FText::FromString(AssetTypeNameStr);
 }
 
 //////////////////////////////////////////////////////////////////////////
diff --git a/Engine/Source/Editor/SourceControlWindows/Private/SSourceControlCommon.h b/Engine/Source/Editor/SourceControlWindows/Private/SSourceControlCommon.h
--- a/Engine/Source/Editor/SourceControlWindows/Private/SSourceControlCommon.h
+++ b/Engine/Source/Editor/SourceControlWindows/Private/SSourceControlCommon.h
@@ -103,6 +103,9 @@ struct IFileViewTreeItem : public IChangelistTreeItem
 	/** The values displayed in the 'Type' column. */
 	virtual const FString& GetType() const { return DefaultStrValue; }
 
+	/** The values displayed in the 'Type Name' column. */
+	virtual const FString& GetTypeName() const { return DefaultStrValue; }
+
 	/** The values displayed in the 'User' column. */
 	virtual const FString& GetCheckedOutBy() const { return DefaultStrValue; }
 
@@ -225,6 +228,7 @@ struct FFileTreeItem : public IFileViewTreeItem
 	virtual const FString& GetName() const override { return AssetNameStr; }
 	virtual const FString& GetPath() const override { return AssetPathStr; }
 	virtual const FString& GetType() const override { return AssetTypeStr; }
+	virtual const FString& GetTypeName() const override { return AssetTypeNameStr; }
 	virtual const FString& GetFullPathname() const override { return FileState->GetFilename(); }
 	virtual const FString& GetCheckedOutBy() const override;
 
@@ -243,6 +247,9 @@ struct FFileTreeItem : public IFileViewTreeItem
 	/** Returns the asset type of the item */
 	FText GetAssetType() const { return AssetType; }
 
+	/** Returns the asset type name of the item */
+	FText GetAssetTypeName() const { return AssetTypeName; }
+
 	/** Returns the asset type color of the item */
 	FSlateColor GetAssetTypeColor() const { return FSlateColor(AssetTypeColor); }
 
@@ -309,6 +316,10 @@ private:
 	FText AssetType;
 	FString AssetTypeStr;
 
+	/** Cached asset type name to display */
+	FText AssetTypeName;
+	FString AssetTypeNameStr;
+
 	/** Cached asset type related color to display */
 	FColor AssetTypeColor;
 
@@ -370,6 +381,8 @@ private:
 	FString AssetPathStr;
 	FText AssetType;
 	FString AssetTypeStr;
+	FText AssetTypeName;
+	FString AssetTypeNameStr;
 	FColor AssetTypeColor;
 };
 
diff --git a/Engine/Source/Editor/SourceControlWindows/SourceControlWindows.Build.cs b/Engine/Source/Editor/SourceControlWindows/SourceControlWindows.Build.cs
--- a/Engine/Source/Editor/SourceControlWindows/SourceControlWindows.Build.cs
+++ b/Engine/Source/Editor/SourceControlWindows/SourceControlWindows.Build.cs
@@ -17,6 +17,7 @@ public class SourceControlWindows : ModuleRules
 				"SourceControl", 
 				"UncontrolledChangelists",
 				"UnsavedAssetsTracker",
+				"AssetDefinition",
 				"AssetTools",
 				"ToolWidgets",
 				"EditorFramework",
