SourceControl: Fixed non package files starting with a dot (such as '.urcignore' or '.gitignore') showing up incorrectly in the 'Submit Files' dialog (they showed up with an empty string for AssetName).

#rnx

diff --git a/Engine/Source/Editor/SourceControlWindows/Private/SSourceControlCommon.cpp b/Engine/Source/Editor/SourceControlWindows/Private/SSourceControlCommon.cpp
--- a/Engine/Source/Editor/SourceControlWindows/Private/SSourceControlCommon.cpp
+++ b/Engine/Source/Editor/SourceControlWindows/Private/SSourceControlCommon.cpp
@@ -23,6 +23,7 @@
 #include "Framework/Docking/TabManager.h"
 #include "Framework/Notifications/NotificationManager.h"
 #include "Logging/MessageLog.h"
+#include "Misc/PathViews.h"
 #include "Misc/ScopedSlowTask.h"
 #include "Modules/ModuleManager.h"
 #include "Editor.h"
@@ -124,6 +125,7 @@ static void RefreshAssetInformationInternal(const TArray<FAssetData>& Assets, co
 {
 	// Initialize display-related members
 	FString Filename = InFilename;
+	FString Extension = FPaths::GetExtension(Filename);
 	FString TempAssetName = SSourceControlCommon::GetDefaultAssetName().ToString();
 	FString TempAssetPath = Filename;
 	FString TempAssetType = SSourceControlCommon::GetDefaultAssetType().ToString();
@@ -135,6 +137,11 @@ static void RefreshAssetInformationInternal(const TArray<FAssetData>& Assets, co
 		127 + FColor::Red.B / 2,
 		200); // Opacity
 
+
+	bool bIsPackageExtension = 
+		FPackageName::IsPackageExtension(*Extension) ||
+		FPackageName::IsVerseExtension(*Extension);
+
 	if (Assets.Num() > 0)
 	{
 		auto IsNotRedirector = [](const FAssetData& InAssetData) { return !InAssetData.IsRedirector(); };
@@ -179,7 +186,7 @@ static void RefreshAssetInformationInternal(const TArray<FAssetData>& Assets, co
 		// Beautify the package name
 		TempPackageName = TempAssetPath + "." + TempAssetName;
 	}
-	else if (FPackageName::TryConvertFilenameToLongPackageName(Filename, TempPackageName))
+	else if (bIsPackageExtension && FPackageName::TryConvertFilenameToLongPackageName(Filename, TempPackageName))
 	{
 		// Fake asset name, asset path from the package name
 		TempAssetPath = TempPackageName;
@@ -194,9 +201,21 @@ static void RefreshAssetInformationInternal(const TArray<FAssetData>& Assets, co
 	else
 	{
 		TempAssetName = FPaths::GetCleanFilename(Filename);
-		TempPackageName = Filename; // put back original package name if the try failed
+		TempPackageName = Filename; // Put back original package name if the try failed
 		TempAssetType = FText::Format(SSourceControlCommon::GetDefaultUnknownAssetType(), FText::FromString(FPaths::GetExtension(Filename).ToUpper())).ToString();
 		TempAssetTypeName = TempAssetType;
+
+		// Attempt to make package name relative to one of the project roots instead of a full absolute path
+		TArray<FSourceControlProjectInfo> CustomProjects = ISourceControlModule::Get().GetCustomProjects();
+		for (const FSourceControlProjectInfo& ProjectInfo : CustomProjects)
+		{
+			FStringView RelativePackageName;
+			if (FPathViews::TryMakeChildPathRelativeTo(TempPackageName, ProjectInfo.ProjectDirectory, RelativePackageName))
+			{
+				TempPackageName = FPaths::Combine(TEXT("/"), FPaths::GetBaseFilename(ProjectInfo.ProjectDirectory), RelativePackageName);
+				break;
+			}
+		}
 	}
 
 	// Finally, assign the temp variables to the member variables
diff --git a/Engine/Source/Runtime/CoreUObject/Private/Misc/PackageName.cpp b/Engine/Source/Runtime/CoreUObject/Private/Misc/PackageName.cpp
--- a/Engine/Source/Runtime/CoreUObject/Private/Misc/PackageName.cpp
+++ b/Engine/Source/Runtime/CoreUObject/Private/Misc/PackageName.cpp
@@ -2567,6 +2567,20 @@ bool FPackageName::IsMapPackageExtension(const TCHAR* Ext)
 	}
 }
 
+bool FPackageName::IsVerseExtension(const TCHAR* Ext)
+{
+	FStringView VerseExtension = TEXTVIEW(".verse");
+	FStringView VModuleExtension = TEXTVIEW(".vmodule");
+	if (*Ext != TEXT('.'))
+	{
+		return (VerseExtension.RightChop(1) == Ext) || (VModuleExtension.RightChop(1) == Ext);
+	}
+	else
+	{
+		return (VerseExtension == Ext) || (VModuleExtension == Ext);
+	}
+}
+
 const TCHAR* FPackageName::GetGeneratedPackageSubPath()
 {
 	return TEXT("_Generated_");
diff --git a/Engine/Source/Runtime/CoreUObject/Public/Misc/PackageName.h b/Engine/Source/Runtime/CoreUObject/Public/Misc/PackageName.h
--- a/Engine/Source/Runtime/CoreUObject/Public/Misc/PackageName.h
+++ b/Engine/Source/Runtime/CoreUObject/Public/Misc/PackageName.h
@@ -600,6 +600,15 @@ public:
 	 */
 	static COREUOBJECT_API bool IsMapPackageExtension(const TCHAR* Ext);
 
+	/**
+	 * Returns whether the passed in extension is a valid verse extension.
+	 * Extensions with and without trailing dots are supported.
+	 *
+	 * @param	Extension to test.
+	 * @return	True if Ext is a verse extension, otherwise false
+	 */
+	static COREUOBJECT_API bool IsVerseExtension(const TCHAR* Ext);
+
 	/** 
 	 * Returns whether the passed in filename ends with any of the known
 	 * package extensions.
