Fix unchanged assets being marked as dirty when checking out unrelated assets

LoadPackage call forces all components to be reregistered which leads to some assets being marked as dirty, even though they were not changed.
Packages from the UnsavedAssets list are already loaded when users attempt to check them out, so use FindPackage instead of LoadPackage.

diff --git a/Engine/Source/Editor/SourceControlWindows/Private/SSourceControlChangelistRows.cpp b/Engine/Source/Editor/SourceControlWindows/Private/SSourceControlChangelistRows.cpp
--- a/Engine/Source/Editor/SourceControlWindows/Private/SSourceControlChangelistRows.cpp
+++ b/Engine/Source/Editor/SourceControlWindows/Private/SSourceControlChangelistRows.cpp
@@ -43,10 +43,17 @@ namespace
 			TArray<UPackage*> Packages;
 			Packages.Reserve(PackageNames.Num());
 			{
-				FGlobalComponentReregisterContext ReregisterContext;
 				for (const FString& PackageName : PackageNames)
 				{
-					Packages.Add(UPackageTools::LoadPackage(PackageName));
+					UPackage* Package = FindPackage(nullptr, *PackageName);
+					if (ensure(Package))
+					{
+						Packages.Add(Package);
+					}
+					else
+					{
+						Packages.Add(UPackageTools::LoadPackage(PackageName));
+					}
 				}
 			}
 			constexpr bool bOnlyDirty = false;
