Fix additional issues causing unchanged assets to be marked as dirty when checking out unrelated assets

Follow up to CL 37121967, which fixed one of the code paths leading to that behavior. This change fixes the remaining code paths that can be encountered when attempting to checkout or save Unsaved Assets.

LoadPackage call forces all components to be reregistered which leads to some assets being marked as dirty, even though they were not changed.
Packages from the UnsavedAssets list are already loaded when users attempt to check them out, so use FindPackage instead of LoadPackage.

diff --git a/Engine/Source/Editor/SourceControlWindows/Private/SSourceControlChangelists.cpp b/Engine/Source/Editor/SourceControlWindows/Private/SSourceControlChangelists.cpp
--- a/Engine/Source/Editor/SourceControlWindows/Private/SSourceControlChangelists.cpp
+++ b/Engine/Source/Editor/SourceControlWindows/Private/SSourceControlChangelists.cpp
@@ -2950,12 +2950,18 @@ TSharedPtr<SWidget> SSourceControlChangelistsWidget::OnOpenContextMenu()
 				TArray<UPackage*> Packages;
 				Packages.Reserve(SelectedFiles.Num());
 				{
-					// Reregister here because otherwise LoadPackage will do it once per call and it can be very slow
-					FGlobalComponentReregisterContext ReregisterContext;
 					for (const FString& Filename : SelectedFiles)
 					{
 						FString PackageName = UPackageTools::FilenameToPackageName(Filename);
-						Packages.Add(UPackageTools::LoadPackage(PackageName));
+						UPackage* Package = FindPackage(nullptr, *PackageName);
+						if (ensure(Package))
+						{
+							Packages.Add(Package);
+						}
+						else
+						{
+							Packages.Add(UPackageTools::LoadPackage(PackageName));
+						}
 					}
 				}
 				const bool bOnlyDirty = false;
@@ -2993,12 +2999,17 @@ TSharedPtr<SWidget> SSourceControlChangelistsWidget::OnOpenContextMenu()
 				TArray<UPackage*> Packages;
 				Packages.Reserve(PackageNames.Num());
 				{
-					// Reregister here because otherwise LoadPackage will do it once per call and it can be very slow
-					FGlobalComponentReregisterContext ReregisterContext;
-
 					for (const FString& PackageName : PackageNames)
 					{
-						Packages.Add(UPackageTools::LoadPackage(PackageName));
+						UPackage* Package = FindPackage(nullptr, *PackageName);
+						if (ensure(Package))
+						{
+							Packages.Add(Package);
+						}
+						else
+						{
+							Packages.Add(UPackageTools::LoadPackage(PackageName));
+						}
 					}
 				}
 				const bool bOnlyDirty = false;
