Added success/failure message when reverting an uncontrolled changelist, to match the UX of reverting a controlled changelist

diff --git a/Engine/Source/Editor/SourceControlWindows/Private/SSourceControlChangelists.cpp b/Engine/Source/Editor/SourceControlWindows/Private/SSourceControlChangelists.cpp
--- a/Engine/Source/Editor/SourceControlWindows/Private/SSourceControlChangelists.cpp
+++ b/Engine/Source/Editor/SourceControlWindows/Private/SSourceControlChangelists.cpp
@@ -2152,7 +2152,14 @@ void SSourceControlChangelistsWidget::OnRevert()
 		{
 			ExecuteUncontrolledChangelistOperation(LOCTEXT("Reverting_Uncontrolled_Files", "Reverting uncontrolled files..."), [&SelectedUncontrolledFiles]()
 			{
-				FUncontrolledChangelistsModule::Get().OnRevert(SelectedUncontrolledFiles);
+				if (FUncontrolledChangelistsModule::Get().OnRevert(SelectedUncontrolledFiles))
+				{
+					SSourceControlCommon::DisplaySourceControlOperationNotification(LOCTEXT("Revert_Uncontrolled_Files_Succeeded", "The selected uncontrolled file(s) were reverted."), SNotificationItem::CS_Success);
+				}
+				else
+				{
+					SSourceControlCommon::DisplaySourceControlOperationNotification(LOCTEXT("Revert_Uncontrolled_Files_Failed", "Failed to revert the selected uncontrolled file(s)."), SNotificationItem::CS_Fail);
+				}
 			});
 		}
 	}
