Source Control: don't show "multiple assets" for files containing a single redirector.

diff --git a/Engine/Source/Editor/SourceControlWindows/Private/SSourceControlCommon.cpp b/Engine/Source/Editor/SourceControlWindows/Private/SSourceControlCommon.cpp
--- a/Engine/Source/Editor/SourceControlWindows/Private/SSourceControlCommon.cpp
+++ b/Engine/Source/Editor/SourceControlWindows/Private/SSourceControlCommon.cpp
@@ -145,42 +145,43 @@ static void RefreshAssetInformationInternal(const TArray<FAssetData>& Assets, co
 	if (Assets.Num() > 0)
 	{
 		auto IsNotRedirector = [](const FAssetData& InAssetData) { return !InAssetData.IsRedirector(); };
-		int32 NumUserFacingAsset = Algo::CountIf(Assets, IsNotRedirector);
+		const int32 NumUserFacingAsset = Algo::CountIf(Assets, IsNotRedirector);
+		const FAssetData* AssetData = nullptr;
 
 		if (NumUserFacingAsset == 1)
 		{
-			const FAssetData& AssetData = *Algo::FindByPredicate(Assets, IsNotRedirector);
+			AssetData = Algo::FindByPredicate(Assets, IsNotRedirector);
+		}
+		else
+		{
+			AssetData = &Assets[0];
+		}
 
-			TempAssetName = RetrieveAssetName(AssetData);
-			TempAssetPath = RetrieveAssetPath(AssetData);
-			TempAssetType = AssetData.AssetClassPath.ToString();
-			TempAssetTypeName = RetrieveAssetTypeName(AssetData);
+		TempAssetName = RetrieveAssetName(*AssetData);
+		TempAssetPath = RetrieveAssetPath(*AssetData);
+		TempAssetTypeName = RetrieveAssetTypeName(*AssetData);
+		TempAssetColor = FColor::White;
 
-			const FAssetToolsModule& AssetToolsModule = FModuleManager::LoadModuleChecked<FAssetToolsModule>(TEXT("AssetTools"));
-			const TSharedPtr<IAssetTypeActions> AssetTypeActions = AssetToolsModule.Get().GetAssetTypeActionsForClass(AssetData.GetClass()).Pin();
+		if (NumUserFacingAsset > 1)
+		{
+			TempAssetType = SSourceControlCommon::GetDefaultMultipleAsset().ToString();
 
-			if (AssetTypeActions.IsValid())
-			{
-				TempAssetColor = AssetTypeActions->GetTypeColor();
-			}
-			else
+			for (int32 i = 1; i < Assets.Num(); ++i)
 			{
-				TempAssetColor = FColor::White;
-			}
+				TempAssetName += TEXT(";") + RetrieveAssetName(Assets[i]);
+			}			
 		}
 		else
 		{
-			TempAssetName = RetrieveAssetName(Assets[0]);
-			TempAssetPath = RetrieveAssetPath(Assets[0]);
-			TempAssetTypeName = RetrieveAssetTypeName(Assets[0]);
+			TempAssetType = AssetData->AssetClassPath.ToString();
 
-			for (int32 i = 1; i < Assets.Num(); ++i)
+			const FAssetToolsModule& AssetToolsModule = FModuleManager::LoadModuleChecked<FAssetToolsModule>(TEXT("AssetTools"));
+			const TSharedPtr<IAssetTypeActions> AssetTypeActions = AssetToolsModule.Get().GetAssetTypeActionsForClass(AssetData->GetClass()).Pin();
+
+			if (AssetTypeActions.IsValid())
 			{
-				TempAssetName += TEXT(";") + RetrieveAssetName(Assets[i]);
+				TempAssetColor = AssetTypeActions->GetTypeColor();
 			}
-
-			TempAssetType = SSourceControlCommon::GetDefaultMultipleAsset().ToString();
-			TempAssetColor = FColor::White;
 		}
 
 		// Beautify the package name
