In SSourceControlChangelistsWidget::OnRefreshUncontrolledChangelistWidgets handle the detection of a duplicate item gracefully instead of with a check.

This applies the same logic that we added to OnRefreshSourceControlWidgets in changelist 33213274 to the uncontrolled changelists.  While we do not know of consistent repro steps for this, a user reported hitting this check and that applying this fix resolved the issue for them.

This also applies the DuplicateItems check to shelved and offline items for completeness of the logic, although there have been no reports related to those cases.

diff --git a/Engine/Source/Editor/SourceControlWindows/Private/SSourceControlChangelists.cpp b/Engine/Source/Editor/SourceControlWindows/Private/SSourceControlChangelists.cpp
--- a/Engine/Source/Editor/SourceControlWindows/Private/SSourceControlChangelists.cpp
+++ b/Engine/Source/Editor/SourceControlWindows/Private/SSourceControlChangelists.cpp
@@ -1152,7 +1152,7 @@ void SSourceControlChangelistsWidget::OnRefreshSourceControlWidgets(int64 CurrUp
 				OnNewSourceControlItemCreated(ShelvedFileState, ShelvedFileItem);
 			}
 
-			if (bChangelistPassedFilter && bShelvedChangelistPassesFilter && FileTextFilter->PassesFilter(*ShelvedFileItem))
+			if (bChangelistPassedFilter && bShelvedChangelistPassesFilter && FileTextFilter->PassesFilter(*ShelvedFileItem) && !DuplicateItems.Contains(ShelvedFileItem.Get()))
 			{
 				ShelvedChangelistTreeItem->AddChild(ShelvedFileItem.ToSharedRef());
 				if (bShelvedChangeslistSelected)
@@ -1178,6 +1178,7 @@ void SSourceControlChangelistsWidget::OnRefreshUncontrolledChangelistWidgets(int
 	const int64 PrevUpdateNum = CurrUpdateNum - 1;
 	int32 DisplayedItemCount = 0;
 	int32 VisitedItemCount = 0;
+	TArray<IChangelistTreeItem*> DuplicateItems;
 
 	TSharedPtr<IChangelistTreeItem> SelectedChangelistItem = UncontrolledChangelistTreeView->GetNumItemsSelected() > 0 ? UncontrolledChangelistTreeView->GetSelectedItems()[0] : TSharedPtr<IChangelistTreeItem>();
 
@@ -1186,11 +1187,19 @@ void SSourceControlChangelistsWidget::OnRefreshUncontrolledChangelistWidgets(int
 		return SelectedChangelistItem && SelectedChangelistItem->GetTreeItemType() == IChangelistTreeItem::UncontrolledChangelist;
 	};
 
-	auto OnUncontrolledChangelistCachedItemVisited = [CurrUpdateNum, &VisitedItemCount](IChangelistTreeItem* Item)
+	auto OnUncontrolledChangelistCachedItemVisited = [CurrUpdateNum, &VisitedItemCount, &DuplicateItems](IChangelistTreeItem* Item)
 	{
-		checkf(Item->VisitedUpdateNum < CurrUpdateNum, TEXT("The same uncontrolled item was visited twice. It is likely present in more than one uncontrolled changelist"));
-		Item->VisitedUpdateNum = CurrUpdateNum;
-		++VisitedItemCount;
+		if (Item->VisitedUpdateNum < CurrUpdateNum)
+		{
+			Item->VisitedUpdateNum = CurrUpdateNum;
+			++VisitedItemCount;
+		}
+		else
+		{
+			// The same uncontrolled item was visited twice. It is likely present in more than one uncontrolled changelist
+			// Handle it gracefully by ignoring the duplicate and then request a full changelists refresh for next tick
+			DuplicateItems.Emplace(Item);
+		}
 	};
 
 	auto OnNewUncontrolledChangelistItemCreated = [this, &OnUncontrolledChangelistCachedItemVisited](TSharedPtr<void> Key, TSharedPtr<IChangelistTreeItem> Item)
@@ -1294,7 +1303,7 @@ void SSourceControlChangelistsWidget::OnRefreshUncontrolledChangelistWidgets(int
 				OnNewUncontrolledChangelistItemCreated(FileState, UncontrolledFileItem);
 			}
 
-			if (bChangelistPassedFilter && FileTextFilter->PassesFilter(*UncontrolledFileItem))
+			if (bChangelistPassedFilter && FileTextFilter->PassesFilter(*UncontrolledFileItem) && !DuplicateItems.Contains(UncontrolledFileItem.Get()))
 			{
 				UncontrolledChangelistNode->AddChild(UncontrolledFileItem.ToSharedRef());
 				if (bUncontrolledChangelistSelected)
@@ -1319,7 +1328,7 @@ void SSourceControlChangelistsWidget::OnRefreshUncontrolledChangelistWidgets(int
 				OnNewOfflineFileCreated(Filename, OfflineFileItem);
 			}
 
-			if (bChangelistPassedFilter && FileTextFilter->PassesFilter(*OfflineFileItem))
+			if (bChangelistPassedFilter && FileTextFilter->PassesFilter(*OfflineFileItem) && !DuplicateItems.Contains(OfflineFileItem.Get()))
 			{
 				UncontrolledChangelistNode->AddChild(OfflineFileItem.ToSharedRef());
 				if (bUncontrolledChangelistSelected)
@@ -1332,6 +1341,9 @@ void SSourceControlChangelistsWidget::OnRefreshUncontrolledChangelistWidgets(int
 
 	RemoveDiscardedUncontrolledChangelistItems();
 	UncontrolledChangelistTreeView->RequestTreeRefresh();
+
+	// If a duplicate item was detected then request a changelists refresh in the next tick 
+	bShouldRefresh |= !DuplicateItems.IsEmpty();
 }
 
 void SSourceControlChangelistsWidget::OnRefreshUI(ERefreshFlags RefreshFlag)
