Fix erroneous errors being logged when attempting to submit a changelist containing deleted files via the editor in a VA enabled project.

- Some methods of invoking a submit in the editor might include deleted files such as "submitting all changes in a project" or using the changelist window but the virtualization process assumes that the caller only requests valid package files on disk be virtualized.
- This problem is very similar to one that was fixed in the UnrealVirtualizationTool in CL 36259815 and to fix it we just need to add similar filtering logic to the editor's revision control submit paths.
-- Filtering ignored files isn't strictly needed but we might as well try to keep the logic the same. There is no point virtualizing a package file that will be ignored by the submit later on.

diff --git a/Engine/Source/Editor/SourceControlWindows/Private/SSourceControlChangelists.cpp b/Engine/Source/Editor/SourceControlWindows/Private/SSourceControlChangelists.cpp
--- a/Engine/Source/Editor/SourceControlWindows/Private/SSourceControlChangelists.cpp
+++ b/Engine/Source/Editor/SourceControlWindows/Private/SSourceControlChangelists.cpp
@@ -2331,15 +2331,9 @@ static bool GetChangelistValidationResult(FSourceControlChangelistPtr InChangeli
 static bool GetOnPresubmitResult(FSourceControlChangelistStatePtr Changelist, FChangeListDescription& Description)
 {
 	const TArray<FSourceControlStateRef>& FileStates = Changelist->GetFilesStates();
-	TArray<FString> LocalFilepathList;
-	LocalFilepathList.Reserve(FileStates.Num());
-	for (const FSourceControlStateRef& State : FileStates)
-	{
-		LocalFilepathList.Add(State->GetFilename());
-	}
 
 	FText FailureMsg;
-	if (!TryToVirtualizeFilesToSubmit(LocalFilepathList, Description.Description, FailureMsg))
+	if (!TryToVirtualizeFilesToSubmit(FileStates, Description.Description, FailureMsg))
 	{
 		// Setup the notification for operation feedback
 		FNotificationInfo Info(FailureMsg);
diff --git a/Engine/Source/Editor/SourceControlWindows/Private/SSourceControlSubmit.cpp b/Engine/Source/Editor/SourceControlWindows/Private/SSourceControlSubmit.cpp
--- a/Engine/Source/Editor/SourceControlWindows/Private/SSourceControlSubmit.cpp
+++ b/Engine/Source/Editor/SourceControlWindows/Private/SSourceControlSubmit.cpp
@@ -43,19 +43,63 @@ bool TryToVirtualizeFilesToSubmit(const TArray<FString>& FilesToSubmit, FText& D
 {
 	using namespace UE::Virtualization;
 
+	IVirtualizationSystem& System = IVirtualizationSystem::Get();
+	if (!System.IsEnabled())
 	{
-		TArray<FText> PayloadErrors;
-		TArray<FText> DescriptionTags;
+		return true; // Early out if VA is not enabled
+	}
 
-PRAGMA_DISABLE_DEPRECATION_WARNINGS
-		ISourceControlModule::Get().GetOnPreSubmitFinalize().Broadcast(FilesToSubmit, DescriptionTags, PayloadErrors);
-PRAGMA_ENABLE_DEPRECATION_WARNINGS
+	TArray<FSourceControlStateRef> FileStates;
+
+	if (ISourceControlModule::Get().GetProvider().GetState(FilesToSubmit, FileStates, EStateCacheUsage::Use) == ECommandResult::Succeeded)
+	{
+		return TryToVirtualizeFilesToSubmit(FileStates, Description, OutFailureMsg);
+	}
+	else
+	{
+		OutFailureMsg = LOCTEXT("SCC_VA_GetStateFailed", "Failed to resolve the file states from revision control!");
+		return false;
 	}
+}
+
+bool TryToVirtualizeFilesToSubmit(const TArray<FSourceControlStateRef>& FileStates, FText& Description, FText& OutFailureMsg)
+{
+	using namespace UE::Virtualization;
 
 	IVirtualizationSystem& System = IVirtualizationSystem::Get();
 	if (!System.IsEnabled())
 	{
-		return true;
+		return true; // Early out if VA is not enabled
+	}
+
+	TArray<FString> FilesToSubmit;
+	{
+		FilesToSubmit.Reserve(FileStates.Num());
+		for (const FSourceControlStateRef& State : FileStates)
+		{
+			if (State->IsDeleted())
+			{
+				UE_LOG(LogVirtualization, Verbose, TEXT("Ignoring package marked for delete '%s'"), *State->GetFilename());
+				continue;
+			}
+
+			if (State->IsIgnored())
+			{
+				UE_LOG(LogVirtualization, Verbose, TEXT("Ignoring package marked for ignore '%s'"), *State->GetFilename());
+				continue;
+			}
+
+			FilesToSubmit.Add(State->GetFilename());
+		}
+	}
+
+	{
+		TArray<FText> PayloadErrors;
+		TArray<FText> DescriptionTags;
+
+		PRAGMA_DISABLE_DEPRECATION_WARNINGS
+		ISourceControlModule::Get().GetOnPreSubmitFinalize().Broadcast(FilesToSubmit, DescriptionTags, PayloadErrors);
+		PRAGMA_ENABLE_DEPRECATION_WARNINGS
 	}
 
 	EVirtualizationOptions VirtualizationOptions = EVirtualizationOptions::None;
diff --git a/Engine/Source/Editor/SourceControlWindows/Private/SSourceControlSubmit.h b/Engine/Source/Editor/SourceControlWindows/Private/SSourceControlSubmit.h
--- a/Engine/Source/Editor/SourceControlWindows/Private/SSourceControlSubmit.h
+++ b/Engine/Source/Editor/SourceControlWindows/Private/SSourceControlSubmit.h
@@ -42,6 +42,7 @@ public:
 };
 
 bool TryToVirtualizeFilesToSubmit(const TArray<FString>& FilesToSubmit, FText& Description, FText& OutFailureMsg);
+bool TryToVirtualizeFilesToSubmit(const TArray<FSourceControlStateRef>& FileStates, FText& Description, FText& OutFailureMsg);
 
 DECLARE_DELEGATE_OneParam(FSourceControlSaveChangelistDescription, const FText& /*NewDescription*/);
 
