Add a new cvar 'RunChoosePackagesToCheckIn' to invoke FSourceControlWindows::ChoosePackagesToCheckIn.

#rnx

- We still have some code paths that can run FSourceControlWindows::ChoosePackagesToCheckIn, such as the 'SourceControlCheckInPrompt' module but it can be tricky to run in the stock editor.
- This cvar will allow us to test the code path quickly.

diff --git a/Engine/Source/Editor/SourceControlWindows/Private/SourceControlWindows.cpp b/Engine/Source/Editor/SourceControlWindows/Private/SourceControlWindows.cpp
--- a/Engine/Source/Editor/SourceControlWindows/Private/SourceControlWindows.cpp
+++ b/Engine/Source/Editor/SourceControlWindows/Private/SourceControlWindows.cpp
@@ -17,11 +17,19 @@
 #include "Widgets/Notifications/SNotificationList.h"
 #include "SourceControlSettings.h"
 #include "Bookmarks/BookmarkScoped.h"
+#include "HAL/IConsoleManager.h"
 
 #if SOURCE_CONTROL_WITH_SLATE
 
 #define LOCTEXT_NAMESPACE "SourceControlWindows"
 
+static FAutoConsoleCommand CCmdRunChoosePackagesToCheckIn = FAutoConsoleCommand(
+	TEXT("RevisionControl.ChoosePackagesToCheckIn"),
+	TEXT("Display all content changes in the current project and allow the user to select which changes to submit to revision control"),
+	FConsoleCommandDelegate::CreateLambda([]()
+		{
+			FSourceControlWindows::ChoosePackagesToCheckIn();
+		}));
 
 //---------------------------------------------------------------------------------------
 // FCheckinResultInfo
@@ -60,7 +68,6 @@ static bool SaveDirtyPackages(bool bUseDialog)
 	return bSaved;
 }
 
-
 //---------------------------------------------------------------------------------------
 // FSourceControlWindows
 
