SourceControl: Fixed difficult to read icons in 'View Changes' dialog.

#rnx

diff --git a/Engine/Source/Developer/SourceControl/Private/RevisionControlStyle/RevisionControlStyle.cpp b/Engine/Source/Developer/SourceControl/Private/RevisionControlStyle/RevisionControlStyle.cpp
--- a/Engine/Source/Developer/SourceControl/Private/RevisionControlStyle/RevisionControlStyle.cpp
+++ b/Engine/Source/Developer/SourceControl/Private/RevisionControlStyle/RevisionControlStyle.cpp
@@ -110,6 +110,7 @@ FDefaultRevisionControlStyle::FDefaultRevisionControlStyle() : FSlateStyleSet(St
 	Set("RevisionControl.Locked", new IMAGE_BRUSH_SVG("Starship/Common/lock", CoreStyleConstants::Icon16x16));
 	Set("RevisionControl.Unlocked", new IMAGE_BRUSH_SVG("Starship/Common/lock-unlocked", CoreStyleConstants::Icon16x16));
 
+	Set("RevisionControl.Shelved", new IMAGE_BRUSH_SVG("Starship/SourceControl/SCC_Shelved", CoreStyleConstants::Icon16x16));
 	Set("RevisionControl.CheckedOut", new IMAGE_BRUSH_SVG("Starship/SourceControl/SCC_CheckedOut", CoreStyleConstants::Icon16x16, StatusCheckedOutColor));
 	
 	Set("RevisionControl.OpenForAdd", new IMAGE_BRUSH_SVG("Starship/SourceControl/RC_MarkedForAdd", CoreStyleConstants::Icon16x16, FStyleColors::AccentBlue));
diff --git a/Engine/Source/Editor/SourceControlWindows/Private/SSourceControlChangelistRows.cpp b/Engine/Source/Editor/SourceControlWindows/Private/SSourceControlChangelistRows.cpp
--- a/Engine/Source/Editor/SourceControlWindows/Private/SSourceControlChangelistRows.cpp
+++ b/Engine/Source/Editor/SourceControlWindows/Private/SSourceControlChangelistRows.cpp
@@ -67,6 +67,10 @@ FName SourceControlFileViewColumn::Icon::Id() { return TEXT("Icon"); }
 FText SourceControlFileViewColumn::Icon::GetDisplayText() {return LOCTEXT("Name_Icon", "Revision Control Status"); }
 FText SourceControlFileViewColumn::Icon::GetToolTipText() { return LOCTEXT("Icon_Column_Tooltip", "Displays the asset/file status"); }
 
+FName SourceControlFileViewColumn::Shelve::Id() { return TEXT("Shelve"); }
+FText SourceControlFileViewColumn::Shelve::GetDisplayText() { return LOCTEXT("Name_Shelve", "Revision Control Shelve"); }
+FText SourceControlFileViewColumn::Shelve::GetToolTipText() { return LOCTEXT("Shelve_Column_Tooltip", "Displays the shelve status"); }
+
 FName SourceControlFileViewColumn::Name::Id() { return TEXT("Name"); }
 FText SourceControlFileViewColumn::Name::GetDisplayText() { return LOCTEXT("Name_Column", "Name"); }
 FText SourceControlFileViewColumn::Name::GetToolTipText() { return LOCTEXT("Name_Column_Tooltip", "Displays the asset/file name"); }
@@ -397,48 +401,78 @@ void SFileTableRow::Construct(const FArguments& InArgs, const TSharedRef<STableV
 
 TSharedRef<SWidget> SFileTableRow::GenerateWidgetForColumn(const FName& ColumnId)
 {
+	auto ConstructPaddedWidget =
+		[](const TSharedRef<SWidget>& Widget)
+		{
+			TSharedRef<SHorizontalBox> HorizontalBox = SNew(SHorizontalBox);
+
+			HorizontalBox->AddSlot()
+				.AutoWidth()
+				.VAlign(VAlign_Center)
+				.Padding(8, 0, 0, 0)
+				[
+					Widget
+				];
+
+			return HorizontalBox;
+		};
+
 	if (ColumnId == SourceControlFileViewColumn::Icon::Id())
 	{
 		return SNew(SBox)
 			.WidthOverride(16) // Small Icons are usually 16x16
 			.HAlign(HAlign_Center)
 			[
-				SSourceControlCommon::GetSCCFileWidget(TreeItem->FileState, TreeItem->IsShelved())
+				SSourceControlCommon::GetSCCStatusWidget(TreeItem->FileState)
+			];
+	}
+	else if (ColumnId == SourceControlFileViewColumn::Shelve::Id())
+	{
+		return SNew(SBox)
+			.WidthOverride(16) // Small Icons are usually 16x16
+			.HAlign(HAlign_Center)
+			[
+				SSourceControlCommon::GetSCCShelveWidget(TreeItem->IsShelved())
 			];
 	}
 	else if (ColumnId == SourceControlFileViewColumn::Name::Id())
 	{
-		return SNew(STextBlock)
+		TSharedPtr<SWidget> Widget = SNew(STextBlock)
 			.Text(this, &SFileTableRow::GetDisplayName)
 			.ToolTipText(this, &SFileTableRow::GetDisplayName)
 			.HighlightText(HighlightText);
+		return ConstructPaddedWidget(Widget.ToSharedRef());
 	}
 	else if (ColumnId == SourceControlFileViewColumn::Path::Id())
 	{
-		return SNew(STextBlock)
+		TSharedPtr<SWidget> Widget = SNew(STextBlock)
 			.Text(this, &SFileTableRow::GetDisplayPath)
 			.ToolTipText(this, &SFileTableRow::GetFilename)
 			.HighlightText(HighlightText);
+		return ConstructPaddedWidget(Widget.ToSharedRef());
 	}
 	else if (ColumnId == SourceControlFileViewColumn::Type::Id())
 	{
-		return SNew(STextBlock)
+		TSharedPtr<SWidget> Widget = SNew(STextBlock)
 			.Text(this, &SFileTableRow::GetDisplayType)
 			.ToolTipText(this, &SFileTableRow::GetDisplayType)
 			.ColorAndOpacity(this, &SFileTableRow::GetDisplayColor)
 			.HighlightText(HighlightText);
+		return ConstructPaddedWidget(Widget.ToSharedRef());
 	}
 	else if (ColumnId == SourceControlFileViewColumn::LastModifiedTimestamp::Id())
 	{
-		return SNew(STextBlock)
+		TSharedPtr<SWidget> Widget = SNew(STextBlock)
 			.ToolTipText(this, &SFileTableRow::GetLastModifiedTimestamp)
 			.Text(this, &SFileTableRow::GetLastModifiedTimestamp);
+		return ConstructPaddedWidget(Widget.ToSharedRef());
 	}
 	else if (ColumnId == SourceControlFileViewColumn::CheckedOutByUser::Id())
 	{
-		return SNew(STextBlock)
+		TSharedPtr<SWidget> Widget = SNew(STextBlock)
 			.ToolTipText(this, &SFileTableRow::GetCheckedOutByUser)
 			.Text(this, &SFileTableRow::GetCheckedOutByUser);
+		return ConstructPaddedWidget(Widget.ToSharedRef());
 	}
 	else
 	{
@@ -562,6 +596,22 @@ void SOfflineFileTableRow::Construct(const FArguments& InArgs, const TSharedRef<
 
 TSharedRef<SWidget> SOfflineFileTableRow::GenerateWidgetForColumn(const FName& ColumnId)
 {
+	auto ConstructPaddedWidget =
+		[](const TSharedRef<SWidget>& Widget)
+		{
+			TSharedRef<SHorizontalBox> HorizontalBox = SNew(SHorizontalBox);
+
+			HorizontalBox->AddSlot()
+				.AutoWidth()
+				.VAlign(VAlign_Center)
+				.Padding(8, 0, 0, 0)
+				[
+					Widget
+				];
+
+			return HorizontalBox;
+		};
+
 	if (ColumnId == SourceControlFileViewColumn::Icon::Id())
 	{
 		return SNew(SBox)
@@ -572,38 +622,52 @@ TSharedRef<SWidget> SOfflineFileTableRow::GenerateWidgetForColumn(const FName& C
 				.Image(FAppStyle::GetBrush(FName("SourceControl.OfflineFile_Small")))
 			];
 	}
+	if (ColumnId == SourceControlFileViewColumn::Shelve::Id())
+	{
+		return SNew(SBox)
+			.WidthOverride(16) // Small Icons are usually 16x16
+			.HAlign(HAlign_Center)
+			[
+				SNullWidget::NullWidget
+			];
+	}
 	else if (ColumnId == SourceControlFileViewColumn::Name::Id())
 	{
-		return SNew(STextBlock)
+		TSharedPtr<SWidget> Widget = SNew(STextBlock)
 			.ToolTipText(this, &SOfflineFileTableRow::GetDisplayName)
 			.Text(this, &SOfflineFileTableRow::GetDisplayName)
 			.HighlightText(HighlightText);
+		return ConstructPaddedWidget(Widget.ToSharedRef());
 	}
 	else if (ColumnId == SourceControlFileViewColumn::Path::Id())
 	{
-		return SNew(STextBlock)
+		TSharedPtr<SWidget> Widget = SNew(STextBlock)
 			.Text(this, &SOfflineFileTableRow::GetDisplayPath)
 			.ToolTipText(this, &SOfflineFileTableRow::GetFilename)
 			.HighlightText(HighlightText);
+		return ConstructPaddedWidget(Widget.ToSharedRef());
 	}
 	else if (ColumnId == SourceControlFileViewColumn::Type::Id())
 	{
-		return SNew(STextBlock)
+		TSharedPtr<SWidget> Widget = SNew(STextBlock)
 			.Text(this, &SOfflineFileTableRow::GetDisplayType)
 			.ToolTipText(this, &SOfflineFileTableRow::GetDisplayType)
 			.ColorAndOpacity(this, &SOfflineFileTableRow::GetDisplayColor)
 			.HighlightText(HighlightText);
+		return ConstructPaddedWidget(Widget.ToSharedRef());
 	}
 	else if (ColumnId == SourceControlFileViewColumn::LastModifiedTimestamp::Id())
 	{
-		return SNew(STextBlock)
+		TSharedPtr<SWidget> Widget = SNew(STextBlock)
 			.ToolTipText(this, &SOfflineFileTableRow::GetLastModifiedTimestamp)
 			.Text(this, &SOfflineFileTableRow::GetLastModifiedTimestamp);
+		return ConstructPaddedWidget(Widget.ToSharedRef());
 	}
 	else if (ColumnId == SourceControlFileViewColumn::CheckedOutByUser::Id())
 	{
-		return SNew(STextBlock)
+		TSharedPtr<SWidget> Widget = SNew(STextBlock)
 			.Text(FText::GetEmpty());
+		return ConstructPaddedWidget(Widget.ToSharedRef());
 	}
 	else if (ColumnId == SourceControlFileViewColumn::Dirty::Id())
 	{
diff --git a/Engine/Source/Editor/SourceControlWindows/Private/SSourceControlChangelistRows.h b/Engine/Source/Editor/SourceControlWindows/Private/SSourceControlChangelistRows.h
--- a/Engine/Source/Editor/SourceControlWindows/Private/SSourceControlChangelistRows.h
+++ b/Engine/Source/Editor/SourceControlWindows/Private/SSourceControlChangelistRows.h
@@ -21,6 +21,14 @@ namespace SourceControlFileViewColumn
 		FText GetToolTipText();
 	};
 
+	/** The shelve column. */
+	namespace Shelve
+	{
+		FName Id();
+		FText GetDisplayText();
+		FText GetToolTipText();
+	};
+
 	/** The file/asset name column. */
 	namespace Name
 	{
diff --git a/Engine/Source/Editor/SourceControlWindows/Private/SSourceControlChangelists.cpp b/Engine/Source/Editor/SourceControlWindows/Private/SSourceControlChangelists.cpp
--- a/Engine/Source/Editor/SourceControlWindows/Private/SSourceControlChangelists.cpp
+++ b/Engine/Source/Editor/SourceControlWindows/Private/SSourceControlChangelists.cpp
@@ -53,6 +53,7 @@
 #include "Framework/MultiBox/MultiBoxBuilder.h"
 #include "Misc/ComparisonUtility.h"
 #include "ProfilingDebugging/CpuProfilerTrace.h"
+#include "RevisionControlStyle/RevisionControlStyle.h"
 
 #define LOCTEXT_NAMESPACE "SourceControlChangelist"
 
@@ -598,7 +599,7 @@ TSharedRef<SWidget> SSourceControlChangelistsWidget::MakeToolBar()
 			NAME_None,
 			LOCTEXT("SourceControl_RefreshButton", "Refresh"),
 			LOCTEXT("SourceControl_RefreshButton_Tooltip", "Refreshes changelists from revision control provider."),
-			FSlateIcon(FAppStyle::GetAppStyleSetName(), "SourceControl.Actions.Refresh"));
+			FSlateIcon(FRevisionControlStyleManager::GetStyleSetName(), "RevisionControl.Actions.Refresh"));
 
 	return ToolBarBuilder.MakeWidget();
 }
@@ -3278,7 +3279,35 @@ TSharedRef<SListView<FChangelistTreeItemPtr>> SSourceControlChangelistsWidget::C
 					.Visibility_Lambda([this](){ return GetColumnSortMode(SourceControlFileViewColumn::Icon::Id()) == EColumnSortMode::None ? EVisibility::Visible : EVisibility::Collapsed; })
 					[
 						SNew(SImage)
-						.Image(FAppStyle::Get().GetBrush("SourceControl.ChangelistsTab"))
+						.Image(FRevisionControlStyleManager::Get().GetBrush("RevisionControl.ChangelistsTab"))
+						.ColorAndOpacity(FSlateColor::UseSubduedForeground())
+					]
+				]
+			]
+
+			+SHeaderRow::Column(SourceControlFileViewColumn::Shelve::Id())
+			.DefaultLabel(SourceControlFileViewColumn::Shelve::GetDisplayText()) // Displayed in the drop down menu to show/hide columns
+			.DefaultTooltip(SourceControlFileViewColumn::Shelve::GetToolTipText())
+			.ShouldGenerateWidget(true) // Ensure the column cannot be hidden (grayed out in the show/hide drop down menu)
+			.FillSized(18)
+			.HeaderContentPadding(FMargin(0))
+			.SortPriority(this, &SSourceControlChangelistsWidget::GetColumnSortPriority, SourceControlFileViewColumn::Shelve::Id())
+			.SortMode(this, &SSourceControlChangelistsWidget::GetColumnSortMode, SourceControlFileViewColumn::Shelve::Id())
+			.OnSort(this, &SSourceControlChangelistsWidget::OnColumnSortModeChanged)
+			[
+				SNew(SHorizontalBox)
+				+SHorizontalBox::Slot()
+				.Padding(1, 0)
+				[
+					SNew(SBox)
+					.WidthOverride(16)
+					.HeightOverride(16)
+					.HAlign(HAlign_Center)
+					.VAlign(VAlign_Center)
+					.Visibility_Lambda([this](){ return GetColumnSortMode(SourceControlFileViewColumn::Shelve::Id()) == EColumnSortMode::None ? EVisibility::Visible : EVisibility::Collapsed; })
+					[
+						SNew(SImage)
+						.Image(FRevisionControlStyleManager::Get().GetBrush("RevisionControl.Shelved"))
 						.ColorAndOpacity(FSlateColor::UseSubduedForeground())
 					]
 				]
@@ -3656,7 +3685,7 @@ void SSourceControlChangelistsWidget::OnColumnSortModeChanged(const EColumnSortP
 		SecondarySortMode = InSortMode;
 	}
 
-	if (IsFileViewSortedByFileStatusIcon())
+	if (IsFileViewSortedByFileStatusIcon() || IsFileViewSortedByFileShelveIcon())
 	{
 		bUpdateMonitoredFileStatusList = true;
 	}
@@ -3670,6 +3699,11 @@ bool SSourceControlChangelistsWidget::IsFileViewSortedByFileStatusIcon() const
 	return PrimarySortedColumn == SourceControlFileViewColumn::Icon::Id() || SecondarySortedColumn == SourceControlFileViewColumn::Icon::Id();
 };
 
+bool SSourceControlChangelistsWidget::IsFileViewSortedByFileShelveIcon() const
+{
+	return PrimarySortedColumn == SourceControlFileViewColumn::Shelve::Id() || SecondarySortedColumn == SourceControlFileViewColumn::Shelve::Id();
+};
+
 bool SSourceControlChangelistsWidget::IsFileViewSortedByLastModifiedTimestamp() const
 {
 	return PrimarySortedColumn == SourceControlFileViewColumn::LastModifiedTimestamp::Id() || SecondarySortedColumn == SourceControlFileViewColumn::LastModifiedTimestamp::Id();
@@ -3691,6 +3725,14 @@ void SSourceControlChangelistsWidget::SortFileView()
 		return LhsVal < RhsVal ? -1 : ( LhsVal == RhsVal ? 0 : 1);
 	};
 
+	auto CompareShelves = [](const IFileViewTreeItem* Lhs, const IFileViewTreeItem* Rhs)
+	{
+		const IChangelistTreeItem::TreeItemType LhsTreeItemType = Lhs->GetTreeItemType();
+		const IChangelistTreeItem::TreeItemType RhsTreeItemType = Rhs->GetTreeItemType();
+
+		return LhsTreeItemType < RhsTreeItemType ? -1 : (LhsTreeItemType == RhsTreeItemType ? 0 : 1);
+	};
+
 	auto CompareNames = [](const IFileViewTreeItem* Lhs, const IFileViewTreeItem* Rhs)
 	{
 		return UE::ComparisonUtility::CompareNaturalOrder(*Lhs->GetName(), *Rhs->GetName());
@@ -3724,6 +3766,10 @@ void SSourceControlChangelistsWidget::SortFileView()
 		{
 			return TFunction<int32(const IFileViewTreeItem*, const IFileViewTreeItem*)>(CompareIcons);
 		}
+		if (ColumnId == SourceControlFileViewColumn::Shelve::Id())
+		{
+			return TFunction<int32(const IFileViewTreeItem*, const IFileViewTreeItem*)>(CompareShelves);
+		}
 		else if (ColumnId == SourceControlFileViewColumn::Name::Id())
 		{
 			return TFunction<int32(const IFileViewTreeItem*, const IFileViewTreeItem*)>(CompareNames);
diff --git a/Engine/Source/Editor/SourceControlWindows/Private/SSourceControlChangelists.h b/Engine/Source/Editor/SourceControlWindows/Private/SSourceControlChangelists.h
--- a/Engine/Source/Editor/SourceControlWindows/Private/SSourceControlChangelists.h
+++ b/Engine/Source/Editor/SourceControlWindows/Private/SSourceControlChangelists.h
@@ -78,6 +78,7 @@ private:
 	EColumnSortMode::Type GetColumnSortMode(const FName ColumnId) const;
 	void OnColumnSortModeChanged(const EColumnSortPriority::Type SortPriority, const FName& ColumnId, const EColumnSortMode::Type InSortMode);
 	bool IsFileViewSortedByFileStatusIcon() const;
+	bool IsFileViewSortedByFileShelveIcon() const;
 	bool IsFileViewSortedByLastModifiedTimestamp() const;
 	void SortFileView();
 
diff --git a/Engine/Source/Editor/SourceControlWindows/Private/SSourceControlCommon.cpp b/Engine/Source/Editor/SourceControlWindows/Private/SSourceControlCommon.cpp
--- a/Engine/Source/Editor/SourceControlWindows/Private/SSourceControlCommon.cpp
+++ b/Engine/Source/Editor/SourceControlWindows/Private/SSourceControlCommon.cpp
@@ -11,6 +11,7 @@
 #include "AssetDefinitionRegistry.h"
 #include "AssetToolsModule.h"
 #include "Styling/AppStyle.h"
+#include "RevisionControlStyle/RevisionControlStyle.h"
 #include "ISourceControlModule.h"
 #include "SourceControlAssetDataCache.h"
 #include "SourceControlHelpers.h"
@@ -379,32 +380,19 @@ void FOfflineFileTreeItem::RefreshAssetInformation()
 namespace SSourceControlCommon
 {
 
-TSharedRef<SWidget> GetSCCFileWidget(FSourceControlStateRef InFileState, bool bIsShelvedFile)
+TSharedRef<SWidget> GetSCCStatusWidget(FSourceControlStateRef InFileState)
 {
-	const FSlateBrush* IconBrush = FAppStyle::GetBrush("ContentBrowser.ColumnViewAssetIcon");
-
-	// Make icon overlays (eg, SCC and dirty status) a reasonable size in relation to the icon size (note: it is assumed this icon is square)
-	const float ICON_SCALING_FACTOR = 0.7f;
-	const float IconOverlaySize = IconBrush->ImageSize.X * ICON_SCALING_FACTOR;
+	const float SizeOverride = 16;
 
 	return SNew(SOverlay)
-		// The actual icon
-		+ SOverlay::Slot()
-		[
-			SNew(SImage)
-			.Image(IconBrush)
-			.ColorAndOpacity_Lambda([bIsShelvedFile]() -> FSlateColor {
-				return FSlateColor(bIsShelvedFile ? FColor::Yellow : FColor::White);
-			})
-		]
 		// Source control state
 		+ SOverlay::Slot()
-		.HAlign(HAlign_Left)
-		.VAlign(VAlign_Top)
+		.HAlign(HAlign_Center)
+		.VAlign(VAlign_Center)
 		[
 			SNew(SBox)
-			.WidthOverride(IconOverlaySize)
-			.HeightOverride(IconOverlaySize)
+			.WidthOverride(SizeOverride)
+			.HeightOverride(SizeOverride)
 			[
 				SNew(SLayeredImage, InFileState->GetIcon())
 				.ToolTipText(InFileState->GetDisplayTooltip())
@@ -412,31 +400,45 @@ TSharedRef<SWidget> GetSCCFileWidget(FSourceControlStateRef InFileState, bool bI
 		];
 }
 
-TSharedRef<SWidget> GetSCCFileWidget()
+TSharedRef<SWidget> GetSCCStatusWidget()
 {
-	const FSlateBrush* IconBrush = FAppStyle::GetBrush("ContentBrowser.ColumnViewAssetIcon");
-
-	// Make icon overlays (eg, SCC and dirty status) a reasonable size in relation to the icon size (note: it is assumed this icon is square)
-	const float ICON_SCALING_FACTOR = 0.7f;
-	const float IconOverlaySize = IconBrush->ImageSize.X * ICON_SCALING_FACTOR;
+	const float SizeOverride = 16;
 
 	return SNew(SOverlay)
-		// The actual icon
+		// Source control state
 		+ SOverlay::Slot()
-		[
-			SNew(SImage)
-			.Image(IconBrush)
-		.ColorAndOpacity(FSlateColor(FColor::White))
-		]
-	// Source control state
-	+ SOverlay::Slot()
-		.HAlign(HAlign_Left)
-		.VAlign(VAlign_Top)
-		[
-			SNew(SBox)
-			.WidthOverride(IconOverlaySize)
-		.HeightOverride(IconOverlaySize)
-		];
+			.HAlign(HAlign_Left)
+			.VAlign(VAlign_Top)
+			[
+				SNew(SBox)
+				.WidthOverride(SizeOverride)
+				.HeightOverride(SizeOverride)
+			];
+}
+
+TSharedRef<SWidget> GetSCCShelveWidget(bool bIsShelvedFile)
+{
+	if (bIsShelvedFile)
+	{
+		const FSlateBrush* IconBrush = FRevisionControlStyleManager::Get().GetBrush("RevisionControl.Shelved");
+
+		return SNew(SOverlay)
+			// Source control shelved
+			+ SOverlay::Slot()
+			[
+				SNew(SImage)
+				.Image(IconBrush)
+				.ColorAndOpacity(FSlateColor::UseSubduedForeground())
+				.ToolTipText(LOCTEXT("SourceControl_Shelved", "Shelved"))
+			];
+	}
+
+	return SNullWidget::NullWidget;
+}
+
+TSharedRef<SWidget> GetSCCShelveWidget()
+{
+	return GetSCCShelveWidget(/*bIsShelvedFile=*/false);
 }
 
 FText GetDefaultAssetName()
diff --git a/Engine/Source/Editor/SourceControlWindows/Private/SSourceControlCommon.h b/Engine/Source/Editor/SourceControlWindows/Private/SSourceControlCommon.h
--- a/Engine/Source/Editor/SourceControlWindows/Private/SSourceControlCommon.h
+++ b/Engine/Source/Editor/SourceControlWindows/Private/SSourceControlCommon.h
@@ -389,8 +389,10 @@ private:
 
 namespace SSourceControlCommon
 {
-	TSharedRef<SWidget> GetSCCFileWidget(FSourceControlStateRef InFileState, bool bIsShelvedFile = false);
-	TSharedRef<SWidget> GetSCCFileWidget();
+	TSharedRef<SWidget> GetSCCStatusWidget(FSourceControlStateRef InFileState);
+	TSharedRef<SWidget> GetSCCStatusWidget();
+	TSharedRef<SWidget> GetSCCShelveWidget(bool bIsShelvedFile);
+	TSharedRef<SWidget> GetSCCShelveWidget();
 	FText GetDefaultAssetName();
 	FText GetDefaultAssetType();
 	FText GetDefaultUnknownAssetType();
@@ -431,11 +433,11 @@ struct FSCCFileDragDropOp : public FDragDropOperation
 		// Offline files won't coexist with Files
 		if (!OfflineFiles.IsEmpty())
 		{
-			return SSourceControlCommon::GetSCCFileWidget();
+			return SSourceControlCommon::GetSCCStatusWidget();
 		}
 
 		FSourceControlStateRef FileState = Files.IsEmpty() ? UncontrolledFiles[0] : Files[0];
-		return SSourceControlCommon::GetSCCFileWidget(MoveTemp(FileState));
+		return SSourceControlCommon::GetSCCStatusWidget(MoveTemp(FileState));
 	}
 
 	TArray<FSourceControlStateRef> Files;
