SourceControl: The 'submit files' dialog now refreshes its UI when the FSourceControlAssetDataCache has fetched the asset data (e.g.: the pretty name of the asset).

#rnx

diff --git a/Engine/Source/Editor/SourceControlWindows/Private/SSourceControlSubmit.cpp b/Engine/Source/Editor/SourceControlWindows/Private/SSourceControlSubmit.cpp
--- a/Engine/Source/Editor/SourceControlWindows/Private/SSourceControlSubmit.cpp
+++ b/Engine/Source/Editor/SourceControlWindows/Private/SSourceControlSubmit.cpp
@@ -679,6 +679,11 @@ FReply SSourceControlSubmitWidget::OnKeyDown( const FGeometry& MyGeometry, const
 	return FReply::Unhandled();
 }
 
+void SSourceControlSubmitWidget::Tick(const FGeometry& AllottedGeometry, const double InCurrentTime, const float InDeltaTime)
+{
+	ISourceControlModule::Get().Tick();
+}
+
 TSharedRef<SWidget> SSourceControlSubmitWidget::GenerateWidgetForItemAndColumn(TSharedPtr<FFileTreeItem> Item, const FName ColumnID) const
 {
 	check(Item.IsValid());
@@ -717,7 +722,7 @@ TSharedRef<SWidget> SSourceControlSubmitWidget::GenerateWidgetForItemAndColumn(T
 			.Padding(RowPadding)
 			[
 				SNew(STextBlock)
-				.Text(Item->GetAssetName())
+				.Text_Lambda([Item]() { return Item->GetAssetName(); })
 			];
 	}
 	else if (ColumnID == SSourceControlSubmitWidgetDefs::ColumnID_FileLabel)
@@ -727,7 +732,7 @@ TSharedRef<SWidget> SSourceControlSubmitWidget::GenerateWidgetForItemAndColumn(T
 			.Padding(RowPadding)
 			[
 				SNew(STextBlock)
-				.Text(Item->GetPackageName())
+				.Text_Lambda([Item]() { return Item->GetPackageName(); })
 				.ToolTipText(Item->GetFileName())
 			];
 	}
diff --git a/Engine/Source/Editor/SourceControlWindows/Private/SSourceControlSubmit.h b/Engine/Source/Editor/SourceControlWindows/Private/SSourceControlSubmit.h
--- a/Engine/Source/Editor/SourceControlWindows/Private/SSourceControlSubmit.h
+++ b/Engine/Source/Editor/SourceControlWindows/Private/SSourceControlSubmit.h
@@ -87,6 +87,9 @@ public:
 	/** Used to intercept Escape key press, and interpret it as cancel */
 	virtual FReply OnKeyDown( const FGeometry& MyGeometry, const FKeyEvent& InKeyEvent ) override;
 
+	/** Used to tick the AssetDataCache via the SourceControlModule when the dialog is open to refresh asset names */
+	virtual void Tick(const FGeometry& AllottedGeometry, const double InCurrentTime, const float InDeltaTime) override;
+
 	/** Get dialog result */
 	ESubmitResults::Type GetResult() { return DialogResult; }
 
