Fixed SSourceControlCheckedOutDialog Asset Class column visibility being controlled by the ShowColumnAssetName argument rather than the ShowColumnAssetClass argument.

diff --git a/Engine/Source/Editor/SourceControlWindows/Private/SSourceControlCheckedOutDialog.cpp b/Engine/Source/Editor/SourceControlWindows/Private/SSourceControlCheckedOutDialog.cpp
--- a/Engine/Source/Editor/SourceControlWindows/Private/SSourceControlCheckedOutDialog.cpp
+++ b/Engine/Source/Editor/SourceControlWindows/Private/SSourceControlCheckedOutDialog.cpp
@@ -74,7 +74,7 @@ void SSourceControlCheckedOutDialog::Construct(const FArguments& InArgs)
 		);
 	}
 
-	if (bShowColumnAssetName)
+	if (bShowColumnAssetClass)
 	{
 		HeaderRowWidget->AddColumn(
 			SHeaderRow::Column(SSourceControlConflictWarningWidgetDefs::ColumnID_AssetClassLabel)
