Added support for showing Verse paths to source control tabs and dialogs.
Updated the conflict dialog to use an icon for its icon column header content, and to support sorting by column.
Reworked the Revert dialog to appear and function similar to the checkout dialog, including converting it to a grid with sortable columns. The previous version just showed package paths which don't have a Verse path.
Updated the submit dialog to use an icon for its icon column header content, show the undetermined checkbox state for the checkbox column header when the list items have mixed state, and use tie breakers when sorting columns where many rows share the same value.
Shows file paths using the platform format.

diff --git a/Engine/Source/Developer/SourceControl/Private/SourceControlHelpers.cpp b/Engine/Source/Developer/SourceControl/Private/SourceControlHelpers.cpp
--- a/Engine/Source/Developer/SourceControl/Private/SourceControlHelpers.cpp
+++ b/Engine/Source/Developer/SourceControl/Private/SourceControlHelpers.cpp
@@ -1785,6 +1785,8 @@ FString USourceControlHelpers::PackageFilename( const UPackage* InPackage )
 TArray<FString> USourceControlHelpers::PackageFilenames( const TArray<UPackage*>& InPackages )
 {
 	TArray<FString> OutNames;
+	OutNames.Reserve(InPackages.Num());
+
 	for (int32 PackageIndex = 0; PackageIndex < InPackages.Num(); PackageIndex++)
 	{
 		OutNames.Add(FPaths::ConvertRelativePathToFull(PackageFilename(InPackages[PackageIndex])));
@@ -1797,6 +1799,8 @@ TArray<FString> USourceControlHelpers::PackageFilenames( const TArray<UPackage*>
 TArray<FString> USourceControlHelpers::PackageFilenames( const TArray<FString>& InPackageNames )
 {
 	TArray<FString> OutNames;
+	OutNames.Reserve(InPackageNames.Num());
+
 	for (int32 PackageIndex = 0; PackageIndex < InPackageNames.Num(); PackageIndex++)
 	{
 		OutNames.Add(FPaths::ConvertRelativePathToFull(PackageFilename_Internal(InPackageNames[PackageIndex])));
@@ -1809,6 +1813,7 @@ TArray<FString> USourceControlHelpers::PackageFilenames( const TArray<FString>&
 TArray<FString> USourceControlHelpers::AbsoluteFilenames( const TArray<FString>& InFileNames )
 {
 	TArray<FString> AbsoluteFiles;
+	AbsoluteFiles.Reserve(InFileNames.Num());
 
 	for (const FString& FileName : InFileNames)
 	{
diff --git a/Engine/Source/Editor/SourceControlWindows/Private/SSourceControlChangelistRows.cpp b/Engine/Source/Editor/SourceControlWindows/Private/SSourceControlChangelistRows.cpp
--- a/Engine/Source/Editor/SourceControlWindows/Private/SSourceControlChangelistRows.cpp
+++ b/Engine/Source/Editor/SourceControlWindows/Private/SSourceControlChangelistRows.cpp
@@ -20,6 +20,7 @@
 #include "RevisionControlStyle/RevisionControlStyle.h"
 #include "Styling/SlateStyleRegistry.h"
 #include "Styling/StarshipCoreStyle.h"
+#include "Misc/ComparisonUtility.h"
 #include "Misc/MessageDialog.h"
 #include "Widgets/Images/SThrobber.h"
 #include "Widgets/Layout/SWidgetSwitcher.h"
@@ -63,6 +64,8 @@ namespace
 	}
 }
 
+FName SourceControlFileViewColumn::CheckBox::Id() { return TEXT("CheckBox"); }
+
 FName SourceControlFileViewColumn::Icon::Id() { return TEXT("Icon"); }
 FText SourceControlFileViewColumn::Icon::GetDisplayText() {return LOCTEXT("Name_Icon", "Revision Control Status"); }
 FText SourceControlFileViewColumn::Icon::GetToolTipText() { return LOCTEXT("Icon_Column_Tooltip", "Displays the asset/file status"); }
@@ -103,6 +106,177 @@ FName SourceControlFileViewColumn::Discard::Id() { return TEXT("Discard"); }
 FText SourceControlFileViewColumn::Discard::GetDisplayText() { return LOCTEXT("Discard_Column", "Discard Unsaved Changes"); }
 FText SourceControlFileViewColumn::Discard::GetToolTipText() { return LOCTEXT("Discard_Column_Tooltip", "Provides option to discard unsaved changes to an asset/file"); }
 
+namespace SourceControlFileViewColumn
+{
+	TFunction<int32(const IFileViewTreeItem&, const IFileViewTreeItem&)> GetColumnComparer(FName ColumnId, EPathFlags PathFlags)
+	{
+		if (ColumnId == SourceControlFileViewColumn::CheckBox::Id())
+		{
+			return [](const IFileViewTreeItem& Lhs, const IFileViewTreeItem& Rhs)
+			{
+				ECheckBoxState LhsVal = Lhs.GetCheckBoxState();
+				ECheckBoxState RhsVal = Rhs.GetCheckBoxState();
+				return LhsVal < RhsVal ? -1 : (LhsVal == RhsVal ? 0 : 1);
+			};
+		}
+		else if (ColumnId == SourceControlFileViewColumn::Icon::Id())
+		{
+			return [](const IFileViewTreeItem& Lhs, const IFileViewTreeItem& Rhs)
+			{
+				int32 LhsVal = Lhs.GetIconSortingPriority();
+				int32 RhsVal = Rhs.GetIconSortingPriority();
+				return LhsVal < RhsVal ? -1 : (LhsVal == RhsVal ? 0 : 1);
+			};
+		}
+		else if (ColumnId == SourceControlFileViewColumn::Shelve::Id())
+		{
+			return [](const IFileViewTreeItem& Lhs, const IFileViewTreeItem& Rhs)
+			{
+				const IChangelistTreeItem::TreeItemType LhsTreeItemType = Lhs.GetTreeItemType();
+				const IChangelistTreeItem::TreeItemType RhsTreeItemType = Rhs.GetTreeItemType();
+
+				return LhsTreeItemType < RhsTreeItemType ? -1 : (LhsTreeItemType == RhsTreeItemType ? 0 : 1);
+			};
+		}
+		else if (ColumnId == SourceControlFileViewColumn::Name::Id())
+		{
+			return [](const IFileViewTreeItem& Lhs, const IFileViewTreeItem& Rhs)
+			{
+				return UE::ComparisonUtility::CompareNaturalOrder(*Lhs.GetName(), *Rhs.GetName());
+			};
+		}
+		else if (ColumnId == SourceControlFileViewColumn::Path::Id())
+		{
+			if (PathFlags == EPathFlags::ShowingPackageName)
+			{
+				return [](const IFileViewTreeItem& Lhs, const IFileViewTreeItem& Rhs)
+				{
+					return Lhs.GetPackageName().Compare(Rhs.GetPackageName());
+				};
+			}
+			else if (PathFlags == EPathFlags::ShowingVersePath)
+			{
+				return [](const IFileViewTreeItem& Lhs, const IFileViewTreeItem& Rhs)
+				{
+					if (Lhs.GetVersePath().IsValid() && Rhs.GetVersePath().IsValid())
+					{
+						return Lhs.GetVersePath().Compare(Rhs.GetVersePath());
+					}
+					else if (Lhs.GetVersePath().IsValid())
+					{
+						return Lhs.GetVersePath().ToString().Compare(Rhs.GetPath());
+					}
+					else if (Rhs.GetVersePath().IsValid())
+					{
+						return Lhs.GetPath().Compare(Rhs.GetVersePath().ToString());
+					}
+					return Lhs.GetPath().Compare(Rhs.GetPath());
+				};
+			}
+			else if (PathFlags == (EPathFlags::ShowingPackageName | EPathFlags::ShowingVersePath))
+			{
+				return [](const IFileViewTreeItem& Lhs, const IFileViewTreeItem& Rhs)
+				{
+					if (Lhs.GetVersePath().IsValid() && Rhs.GetVersePath().IsValid())
+					{
+						return Lhs.GetVersePath().Compare(Rhs.GetVersePath());
+					}
+					else if (Lhs.GetVersePath().IsValid())
+					{
+						return Lhs.GetVersePath().ToString().Compare(Rhs.GetPackageName());
+					}
+					else if (Rhs.GetVersePath().IsValid())
+					{
+						return Lhs.GetPackageName().Compare(Rhs.GetVersePath().ToString());
+					}
+					return Lhs.GetPackageName().Compare(Rhs.GetPackageName());
+				};
+			}
+
+			return [](const IFileViewTreeItem& Lhs, const IFileViewTreeItem& Rhs)
+			{
+				return Lhs.GetPath().Compare(Rhs.GetPath());
+			};
+		}
+		else if (ColumnId == SourceControlFileViewColumn::Type::Id())
+		{
+			return [](const IFileViewTreeItem& Lhs, const IFileViewTreeItem& Rhs)
+			{
+				return FCString::Stricmp(*Lhs.GetType(), *Rhs.GetType());
+			};
+		}
+		else if (ColumnId == SourceControlFileViewColumn::LastModifiedTimestamp::Id())
+		{
+			return [](const IFileViewTreeItem& Lhs, const IFileViewTreeItem& Rhs)
+			{
+				const FDateTime& LhsVal = Lhs.GetLastModifiedDateTime();
+				const FDateTime& RhsVal = Rhs.GetLastModifiedDateTime();
+				return LhsVal < RhsVal ? -1 : (LhsVal == RhsVal ? 0 : 1);
+			};
+		}
+		else if (ColumnId == SourceControlFileViewColumn::CheckedOutByUser::Id())
+		{
+			return [](const IFileViewTreeItem& Lhs, const IFileViewTreeItem& Rhs)
+			{
+				return FCString::Stricmp(*Lhs.GetCheckedOutBy(), *Rhs.GetCheckedOutBy());
+			};
+		}
+		else
+		{
+			checkNoEntry();
+			return [](const IFileViewTreeItem& Lhs, const IFileViewTreeItem& Rhs)
+			{
+				return &Lhs < &Rhs ? -1 : (&Lhs == &Rhs ? 0 : 1);
+			};
+		};
+	}
+
+	TFunction<bool(const IFileViewTreeItem&, const IFileViewTreeItem&)> GetSortPredicate(EColumnSortMode::Type SortMode, FName ColumnId, EPathFlags PathFlags)
+	{
+		int32 Sign;
+		if (SortMode == EColumnSortMode::Ascending)
+		{
+			Sign = 1;
+		}
+		else if (SortMode == EColumnSortMode::Descending)
+		{
+			Sign = -1;
+		}
+		else
+		{
+			return {};
+		}
+
+		TArray<TFunction<int32(const IFileViewTreeItem&, const IFileViewTreeItem&)>> ColumnComparers;
+		ColumnComparers.Reserve(3);
+
+		ColumnComparers.Add(GetColumnComparer(ColumnId, PathFlags));
+
+		if (ColumnId != Path::Id())
+		{
+			// Use name and path as tie breakers.
+			if (ColumnId != Name::Id())
+			{
+				ColumnComparers.Add(GetColumnComparer(Name::Id(), PathFlags));
+			}
+			ColumnComparers.Add(GetColumnComparer(Path::Id(), PathFlags));
+		}
+
+		return [Sign, ColumnComparers = MoveTemp(ColumnComparers)](const IFileViewTreeItem& Lhs, const IFileViewTreeItem& Rhs)
+		{
+			for (const TFunction<int32(const IFileViewTreeItem&, const IFileViewTreeItem&)>& ColumnComparer : ColumnComparers)
+			{
+				const int32 Result = ColumnComparer(Lhs, Rhs);
+				if (Result != 0)
+				{
+					return Sign * Result < 0;
+				}
+			}
+			return false;
+		};
+	}
+}
+
 FText FormatChangelistFileCountText(int32 DisplayedCount, int32 TotalCount)
 {
 	return DisplayedCount == TotalCount ?
@@ -393,6 +567,8 @@ void SFileTableRow::Construct(const FArguments& InArgs, const TSharedRef<STableV
 
 	HighlightText = InArgs._HighlightText;
 
+	PathFlags = InArgs._PathFlags;
+
 	FSuperRowType::FArguments Args = FSuperRowType::FArguments()
 		.OnDragDetected(InArgs._OnDragDetected)
 		.ShowSelection(true);
@@ -417,7 +593,18 @@ TSharedRef<SWidget> SFileTableRow::GenerateWidgetForColumn(const FName& ColumnId
 			return HorizontalBox;
 		};
 
-	if (ColumnId == SourceControlFileViewColumn::Icon::Id())
+	if (ColumnId == SourceControlFileViewColumn::CheckBox::Id())
+	{
+		return SNew(SHorizontalBox)
+			+SHorizontalBox::Slot()
+			.Padding(FMargin(10.0f, 3.0f, 6.0f, 3.0f))
+			[
+				SNew(SCheckBox)
+				.IsChecked(this, &SFileTableRow::GetCheckBoxState)
+				.OnCheckStateChanged(this, &SFileTableRow::SetCheckBoxState)
+			];
+	}
+	else if (ColumnId == SourceControlFileViewColumn::Icon::Id())
 	{
 		return SNew(SBox)
 			.WidthOverride(16) // Small Icons are usually 16x16
@@ -480,15 +667,36 @@ TSharedRef<SWidget> SFileTableRow::GenerateWidgetForColumn(const FName& ColumnId
 	}
 }
 
-void SFileTableRow::PopulateSearchString(const FFileTreeItem& Item, TArray<FString>& OutStrings)
+void SFileTableRow::PopulateSearchString(const FFileTreeItem& Item, SourceControlFileViewColumn::EPathFlags PathFlags, TArray<FString>& OutStrings)
 {
 	OutStrings.Emplace(Item.GetAssetName().ToString()); // Name
-	OutStrings.Emplace(Item.GetAssetPath().ToString()); // Path
+	if (EnumHasAnyFlags(PathFlags, SourceControlFileViewColumn::EPathFlags::ShowingVersePath) && Item.GetVersePath().IsValid())
+	{
+		OutStrings.Emplace(Item.GetAssetVersePath().ToString());
+	}
+	else if (EnumHasAnyFlags(PathFlags, SourceControlFileViewColumn::EPathFlags::ShowingPackageName))
+	{
+		OutStrings.Emplace(Item.GetAssetPackageName().ToString()); // Path
+	}
+	else
+	{
+		OutStrings.Emplace(Item.GetAssetPath().ToString()); // Path
+	}
 	OutStrings.Emplace(Item.GetAssetType().ToString()); // Type
 	OutStrings.Emplace(Item.GetLastModifiedTimestamp().ToString());
 	OutStrings.Emplace(Item.GetCheckedOutByUser().ToString());
 }
 
+ECheckBoxState SFileTableRow::GetCheckBoxState() const
+{
+	return TreeItem->GetCheckBoxState();
+}
+
+void SFileTableRow::SetCheckBoxState(ECheckBoxState NewState)
+{
+	TreeItem->SetCheckBoxState(NewState);
+}
+
 FText SFileTableRow::GetDisplayName() const
 {
 	return TreeItem->GetAssetName();
@@ -501,6 +709,14 @@ FText SFileTableRow::GetFilename() const
 
 FText SFileTableRow::GetDisplayPath() const
 {
+	if (EnumHasAnyFlags(PathFlags.Get(SourceControlFileViewColumn::EPathFlags::Default), SourceControlFileViewColumn::EPathFlags::ShowingVersePath) && TreeItem->GetVersePath().IsValid())
+	{
+		return TreeItem->GetAssetVersePath();
+	}
+	else if (EnumHasAnyFlags(PathFlags.Get(SourceControlFileViewColumn::EPathFlags::Default), SourceControlFileViewColumn::EPathFlags::ShowingPackageName))
+	{
+		return TreeItem->GetAssetPackageName();
+	}
 	return TreeItem->GetAssetPath();
 }
 
@@ -587,6 +803,7 @@ void SOfflineFileTableRow::Construct(const FArguments& InArgs, const TSharedRef<
 {
 	TreeItem = static_cast<FOfflineFileTreeItem*>(InArgs._TreeItemToVisualize.Get());
 	HighlightText = InArgs._HighlightText;
+	PathFlags = InArgs._PathFlags;
 
 	FSuperRowType::FArguments Args = FSuperRowType::FArguments()
 		.OnDragDetected(InArgs._OnDragDetected)
@@ -612,11 +829,23 @@ TSharedRef<SWidget> SOfflineFileTableRow::GenerateWidgetForColumn(const FName& C
 			return HorizontalBox;
 		};
 
-	if (ColumnId == SourceControlFileViewColumn::Icon::Id())
+	if (ColumnId == SourceControlFileViewColumn::CheckBox::Id())
+	{
+		return SNew(SHorizontalBox)
+			+SHorizontalBox::Slot()
+			.Padding(FMargin(10.0f, 3.0f, 6.0f, 3.0f))
+			[
+				SNew(SCheckBox)
+				.IsChecked(this, &SOfflineFileTableRow::GetCheckBoxState)
+				.OnCheckStateChanged(this, &SOfflineFileTableRow::SetCheckBoxState)
+			];
+	}
+	else if (ColumnId == SourceControlFileViewColumn::Icon::Id())
 	{
 		return SNew(SBox)
 			.WidthOverride(16) // Small Icons are usually 16x16
 			.HAlign(HAlign_Center)
+			.VAlign(VAlign_Center)
 			[
 				SNew(SImage)
 				.Image(FAppStyle::GetBrush(FName("SourceControl.OfflineFile_Small")))
@@ -704,7 +933,7 @@ TSharedRef<SWidget> SOfflineFileTableRow::GenerateWidgetForColumn(const FName& C
 					// Normalize packagenames and filenames
 					FString PackageName;
 					{
-						FString TreeName = TreeItem->GetPackageName().ToString();
+						FString TreeName = TreeItem->GetPackageName();
 
 						if (!FPackageName::TryConvertFilenameToLongPackageName(TreeName, PackageName))
 						{
@@ -728,7 +957,7 @@ TSharedRef<SWidget> SOfflineFileTableRow::GenerateWidgetForColumn(const FName& C
 					ExecuteOnGameThread(UE_SOURCE_LOCATION,
 						[this]
 						{
-							TArray<FString> PackageToReload { TreeItem->GetPackageName().ToString() };
+							TArray<FString> PackageToReload { TreeItem->GetPackageName() };
 							const bool bAllowReloadWorld = true;
 							const bool bInteractive = false;
 							USourceControlHelpers::ApplyOperationAndReloadPackages(
@@ -765,14 +994,35 @@ TSharedRef<SWidget> SOfflineFileTableRow::GenerateWidgetForColumn(const FName& C
 	}
 }
 
-void SOfflineFileTableRow::PopulateSearchString(const FOfflineFileTreeItem& Item, TArray<FString>& OutStrings)
+void SOfflineFileTableRow::PopulateSearchString(const FOfflineFileTreeItem& Item, SourceControlFileViewColumn::EPathFlags PathFlags, TArray<FString>& OutStrings)
 {
 	OutStrings.Emplace(Item.GetDisplayName().ToString()); // Name
-	OutStrings.Emplace(Item.GetDisplayPath().ToString()); // Path
+	if (EnumHasAnyFlags(PathFlags, SourceControlFileViewColumn::EPathFlags::ShowingVersePath) && Item.GetVersePath().IsValid())
+	{
+		OutStrings.Emplace(Item.GetDisplayVersePath().ToString());
+	}
+	else if (EnumHasAnyFlags(PathFlags, SourceControlFileViewColumn::EPathFlags::ShowingPackageName))
+	{
+		OutStrings.Emplace(Item.GetDisplayPackageName().ToString()); // Path
+	}
+	else
+	{
+		OutStrings.Emplace(Item.GetDisplayPath().ToString()); // Path
+	}
 	OutStrings.Emplace(Item.GetDisplayType().ToString()); // Type
 	OutStrings.Emplace(Item.GetLastModifiedTimestamp().ToString());
 }
 
+ECheckBoxState SOfflineFileTableRow::GetCheckBoxState() const
+{
+	return TreeItem->GetCheckBoxState();
+}
+
+void SOfflineFileTableRow::SetCheckBoxState(ECheckBoxState NewState)
+{
+	TreeItem->SetCheckBoxState(NewState);
+}
+
 FText SOfflineFileTableRow::GetDisplayName() const
 {
 	return TreeItem->GetDisplayName();
@@ -780,11 +1030,19 @@ FText SOfflineFileTableRow::GetDisplayName() const
 
 FText SOfflineFileTableRow::GetFilename() const
 {
-	return TreeItem->GetPackageName();
+	return TreeItem->GetDisplayPackageName();
 }
 
 FText SOfflineFileTableRow::GetDisplayPath() const
 {
+	if (EnumHasAnyFlags(PathFlags.Get(SourceControlFileViewColumn::EPathFlags::Default), SourceControlFileViewColumn::EPathFlags::ShowingVersePath) && TreeItem->GetVersePath().IsValid())
+	{
+		return TreeItem->GetDisplayVersePath();
+	}
+	else if (EnumHasAnyFlags(PathFlags.Get(SourceControlFileViewColumn::EPathFlags::Default), SourceControlFileViewColumn::EPathFlags::ShowingPackageName))
+	{
+		return TreeItem->GetDisplayPackageName();
+	}
 	return TreeItem->GetDisplayPath();
 }
 
diff --git a/Engine/Source/Editor/SourceControlWindows/Private/SSourceControlChangelistRows.h b/Engine/Source/Editor/SourceControlWindows/Private/SSourceControlChangelistRows.h
--- a/Engine/Source/Editor/SourceControlWindows/Private/SSourceControlChangelistRows.h
+++ b/Engine/Source/Editor/SourceControlWindows/Private/SSourceControlChangelistRows.h
@@ -4,6 +4,7 @@
 
 #include "CoreMinimal.h"
 
+#include "Misc/EnumClassFlags.h"
 #include "Widgets/SCompoundWidget.h"
 #include "Widgets/Views/STableRow.h"
 #include "Widgets/Views/STreeView.h"
@@ -13,13 +14,19 @@
 /** Lists the unique columns used in the list view displaying controlled/uncontrolled changelist files. */
 namespace SourceControlFileViewColumn
 {
+	/** The checkbox column. */
+	namespace CheckBox
+	{
+		FName Id();
+	}
+
 	/** The icon column. */
 	namespace Icon
 	{
 		FName Id();
 		FText GetDisplayText();
 		FText GetToolTipText();
-	};
+	}
 
 	/** The shelve column. */
 	namespace Shelve
@@ -27,7 +34,7 @@ namespace SourceControlFileViewColumn
 		FName Id();
 		FText GetDisplayText();
 		FText GetToolTipText();
-	};
+	}
 
 	/** The file/asset name column. */
 	namespace Name
@@ -35,7 +42,7 @@ namespace SourceControlFileViewColumn
 		FName Id();
 		FText GetDisplayText();
 		FText GetToolTipText();
-	};
+	}
 	
 	/** The file/asset path column. */
 	namespace Path
@@ -43,7 +50,7 @@ namespace SourceControlFileViewColumn
 		FName Id();
 		FText GetDisplayText();
 		FText GetToolTipText();
-	};
+	}
 
 	/** The file/asset type column. */
 	namespace Type
@@ -51,7 +58,7 @@ namespace SourceControlFileViewColumn
 		FName Id();
 		FText GetDisplayText();
 		FText GetToolTipText();
-	};
+	}
 
 	/** The last time the file/asset was modified column. */
 	namespace LastModifiedTimestamp
@@ -92,6 +99,19 @@ namespace SourceControlFileViewColumn
 		FText GetDisplayText();
 		FText GetToolTipText();
 	}
+
+	enum class EPathFlags
+	{
+		Default = 0,
+
+		ShowingPackageName = 1 << 0,
+		ShowingVersePath = 1 << 1,
+	};
+	ENUM_CLASS_FLAGS(EPathFlags)
+
+	TFunction<int32(const IFileViewTreeItem&, const IFileViewTreeItem&)> GetColumnComparer(FName ColumnId, EPathFlags PathFlags);
+
+	TFunction<bool(const IFileViewTreeItem&, const IFileViewTreeItem&)> GetSortPredicate(EColumnSortMode::Type, FName ColumnId, EPathFlags PathFlags);
 }
 
 
@@ -195,10 +215,12 @@ public:
 	SLATE_BEGIN_ARGS(SFileTableRow)
 		: _TreeItemToVisualize(nullptr)
 		, _HighlightText()
+		, _PathFlags(SourceControlFileViewColumn::EPathFlags::Default)
 	{
 	}
 		SLATE_ARGUMENT(FChangelistTreeItemPtr, TreeItemToVisualize)
 		SLATE_ATTRIBUTE(FText, HighlightText)
+		SLATE_ATTRIBUTE(SourceControlFileViewColumn::EPathFlags, PathFlags)
 		SLATE_EVENT(FOnDragDetected, OnDragDetected)
 	SLATE_END_ARGS()
 
@@ -213,6 +235,8 @@ public:
 	// SMultiColumnTableRow overrides
 	virtual TSharedRef<SWidget> GenerateWidgetForColumn(const FName& ColumnName) override;
 
+	ECheckBoxState GetCheckBoxState() const;
+	void SetCheckBoxState(ECheckBoxState NewState);
 	FText GetDisplayName() const;
 	FText GetFilename() const;
 	FText GetDisplayPath() const;
@@ -221,7 +245,7 @@ public:
 	FText GetLastModifiedTimestamp() const;
 	FText GetCheckedOutByUser() const;
 
-	static void PopulateSearchString(const FFileTreeItem& Item, TArray<FString>& OutStrings);
+	static void PopulateSearchString(const FFileTreeItem& Item, SourceControlFileViewColumn::EPathFlags PathFlags, TArray<FString>& OutStrings);
 
 protected:
 	//~ Begin STableRow Interface.
@@ -235,6 +259,8 @@ private:
 
 	/** The text to highlight*/
 	TAttribute<FText> HighlightText;
+
+	TAttribute<SourceControlFileViewColumn::EPathFlags> PathFlags;
 };
 
 
@@ -274,6 +300,7 @@ public:
 	}
 		SLATE_ARGUMENT(FChangelistTreeItemPtr, TreeItemToVisualize)
 		SLATE_ATTRIBUTE(FText, HighlightText)
+		SLATE_ATTRIBUTE(SourceControlFileViewColumn::EPathFlags, PathFlags)
 		SLATE_EVENT(FOnDragDetected, OnDragDetected)
 	SLATE_END_ARGS()
 
@@ -288,15 +315,16 @@ public:
 	// SMultiColumnTableRow overrides
 	virtual TSharedRef<SWidget> GenerateWidgetForColumn(const FName& ColumnName) override;
 
+	ECheckBoxState GetCheckBoxState() const;
+	void SetCheckBoxState(ECheckBoxState NewState);
 	FText GetDisplayName() const;
 	FText GetFilename() const;
 	FText GetDisplayPath() const;
 	FText GetDisplayType() const;
 	FSlateColor GetDisplayColor() const;
 	FText GetLastModifiedTimestamp() const;
-	FText GetCheckedOutByUser() const;
 
-	static void PopulateSearchString(const FOfflineFileTreeItem& Item, TArray<FString>& OutStrings);
+	static void PopulateSearchString(const FOfflineFileTreeItem& Item, SourceControlFileViewColumn::EPathFlags PathFlags, TArray<FString>& OutStrings);
 
 protected:
 	//~ Begin STableRow Interface.
@@ -310,4 +338,6 @@ private:
 
 	/** The text to highlight*/
 	TAttribute<FText> HighlightText;
+
+	TAttribute<SourceControlFileViewColumn::EPathFlags> PathFlags;
 };
diff --git a/Engine/Source/Editor/SourceControlWindows/Private/SSourceControlChangelists.cpp b/Engine/Source/Editor/SourceControlWindows/Private/SSourceControlChangelists.cpp
--- a/Engine/Source/Editor/SourceControlWindows/Private/SSourceControlChangelists.cpp
+++ b/Engine/Source/Editor/SourceControlWindows/Private/SSourceControlChangelists.cpp
@@ -18,9 +18,11 @@
 #include "Widgets/Layout/SExpandableArea.h"
 #include "Widgets/Layout/SWidgetSwitcher.h"
 
+#include "IAssetTools.h"
 #include "ISourceControlProvider.h"
 #include "ISourceControlModule.h"
 #include "ISourceControlWindowsModule.h"
+#include "Interfaces/IPluginManager.h"
 #include "UncontrolledChangelistsModule.h"
 #include "SourceControlOperations.h"
 #include "SSourceControlChangelistRows.h"
@@ -51,7 +53,6 @@
 #include "Framework/Notifications/NotificationManager.h"
 #include "Framework/Docking/TabManager.h"
 #include "Framework/MultiBox/MultiBoxBuilder.h"
-#include "Misc/ComparisonUtility.h"
 #include "ProfilingDebugging/CpuProfilerTrace.h"
 #include "RevisionControlStyle/RevisionControlStyle.h"
 
@@ -351,6 +352,8 @@ void SSourceControlChangelistsWidget::FAsyncTimestampUpdater::DoWork()
 
 void SSourceControlChangelistsWidget::Construct(const FArguments& InArgs)
 {
+	bShowingContentVersePath = FAssetToolsModule::GetModule().Get().ShowingContentVersePath();
+
 	// Register delegates
 	ISourceControlModule& SCCModule = ISourceControlModule::Get();
 	FUncontrolledChangelistsModule& UncontrolledChangelistModule = FUncontrolledChangelistsModule::Get();
@@ -362,6 +365,8 @@ void SSourceControlChangelistsWidget::Construct(const FArguments& InArgs)
 	// Monitor when packages are saved to update the timestamps.
 	UPackage::PackageSavedWithContextEvent.AddSP(this, &SSourceControlChangelistsWidget::OnPackageSaved);
 
+	IPluginManager::Get().OnPluginEdited().AddSP(this, &SSourceControlChangelistsWidget::OnPluginEdited);
+
 	// No default sorting, sorting make the view slower, just pay the cost if user wants it.
 	PrimarySortedColumn = FName();
 
@@ -661,6 +666,8 @@ void SSourceControlChangelistsWidget::Tick(const FGeometry& AllottedGeometry, co
 		bShouldRefresh = true;
 	}
 
+	bool bSortFileView = false;
+
 	if (bShouldRefresh)
 	{
 		if (!bInitialRefreshDone)
@@ -726,6 +733,26 @@ void SSourceControlChangelistsWidget::Tick(const FGeometry& AllottedGeometry, co
 		}
 	}
 
+	{
+		const bool bNewShowingContentVersePath = FAssetToolsModule::GetModule().Get().ShowingContentVersePath();
+		if (bShowingContentVersePath != bNewShowingContentVersePath)
+		{
+			bShowingContentVersePath = bNewShowingContentVersePath;
+			bSortFileView = true;
+		}
+	}
+
+	if (bShouldRefreshVersePaths && bShowingContentVersePath)
+	{
+		bShouldRefreshVersePaths = false;
+
+		for (const TSharedPtr<IChangelistTreeItem>& Item : FileListNodes)
+		{
+			// Resort of any of the Verse paths change.
+			bSortFileView |= static_cast<IFileViewTreeItem*>(Item.Get())->RefreshVersePath();
+		}
+	}
+
 	// If requested timestamps were gathered in background.
 	if (TimestampUpdateTask && TimestampUpdateTask->IsDone())
 	{
@@ -769,10 +796,16 @@ void SSourceControlChangelistsWidget::Tick(const FGeometry& AllottedGeometry, co
 
 		if (bUpdatedItemInView && IsFileViewSortedByLastModifiedTimestamp())
 		{
-			SortFileView();
-			GetActiveFileListView().RequestListRefresh();
+			bSortFileView = true;
 		}
 	}
+
+	if (bSortFileView)
+	{
+		SortFileView();
+		GetActiveFileListView().RequestListRefresh();
+	}
+
 	if (!TimestampUpdateTask && OutdatedTimestampFiles.Num()) // Should be sufficient to run only one async task at the time.
 	{
 		// Query the timestamps asynchronously.
@@ -3486,6 +3519,7 @@ TSharedRef<ITableRow> SSourceControlChangelistsWidget::OnGenerateRow(TSharedPtr<
 		return SNew(SFileTableRow, OwnerTable)
 			.TreeItemToVisualize(InTreeItem)
 			.HighlightText_Lambda([this]() { return FileSearchBox->GetText(); })
+			.PathFlags_Lambda([this]() { return bShowingContentVersePath ? SourceControlFileViewColumn::EPathFlags::ShowingVersePath : SourceControlFileViewColumn::EPathFlags::Default; })
 			.OnDragDetected(this, &SSourceControlChangelistsWidget::OnFilesDragged);
 
 	case IChangelistTreeItem::OfflineFile:
@@ -3493,6 +3527,7 @@ TSharedRef<ITableRow> SSourceControlChangelistsWidget::OnGenerateRow(TSharedPtr<
 		return SNew(SOfflineFileTableRow, OwnerTable)
 			.TreeItemToVisualize(InTreeItem)
 			.HighlightText_Lambda([this]() { return FileSearchBox->GetText(); })
+			.PathFlags_Lambda([this]() { return bShowingContentVersePath ? SourceControlFileViewColumn::EPathFlags::ShowingVersePath : SourceControlFileViewColumn::EPathFlags::Default; })
 			.OnDragDetected(this, &SSourceControlChangelistsWidget::OnUnsavedAssetsDragged);
 
 	case IChangelistTreeItem::ShelvedChangelist:
@@ -3504,7 +3539,8 @@ TSharedRef<ITableRow> SSourceControlChangelistsWidget::OnGenerateRow(TSharedPtr<
 		bUpdateMonitoredFileStatusList = true;
 		return SNew(SFileTableRow, OwnerTable)
 			.TreeItemToVisualize(InTreeItem)
-			.HighlightText_Lambda([this]() { return FileSearchBox->GetText(); });
+			.HighlightText_Lambda([this]() { return FileSearchBox->GetText(); })
+			.PathFlags_Lambda([this]() { return bShowingContentVersePath ? SourceControlFileViewColumn::EPathFlags::ShowingVersePath : SourceControlFileViewColumn::EPathFlags::Default; });
 
 	case IChangelistTreeItem::UnsavedAssets:
 		return SNew(SUnsavedAssetsTableRow, OwnerTable);
@@ -3741,90 +3777,14 @@ void SSourceControlChangelistsWidget::SortFileView()
 		return; // No column selected for sorting or nothing to sort.
 	}
 
-	auto CompareIcons = [](const IFileViewTreeItem* Lhs, const IFileViewTreeItem* Rhs)
-	{
-		int32 LhsVal = Lhs->GetIconSortingPriority();
-		int32 RhsVal = Rhs->GetIconSortingPriority();
-		return LhsVal < RhsVal ? -1 : ( LhsVal == RhsVal ? 0 : 1);
-	};
-
-	auto CompareShelves = [](const IFileViewTreeItem* Lhs, const IFileViewTreeItem* Rhs)
-	{
-		const IChangelistTreeItem::TreeItemType LhsTreeItemType = Lhs->GetTreeItemType();
-		const IChangelistTreeItem::TreeItemType RhsTreeItemType = Rhs->GetTreeItemType();
-
-		return LhsTreeItemType < RhsTreeItemType ? -1 : (LhsTreeItemType == RhsTreeItemType ? 0 : 1);
-	};
-
-	auto CompareNames = [](const IFileViewTreeItem* Lhs, const IFileViewTreeItem* Rhs)
-	{
-		return UE::ComparisonUtility::CompareNaturalOrder(*Lhs->GetName(), *Rhs->GetName());
-	};
-
-	auto ComparePaths = [](const IFileViewTreeItem* Lhs, const IFileViewTreeItem* Rhs)
-	{
-		return FCString::Stricmp(*Lhs->GetPath(), *Rhs->GetPath());
-	};
+	const SourceControlFileViewColumn::EPathFlags PathFlags =
+		bShowingContentVersePath ? SourceControlFileViewColumn::EPathFlags::ShowingVersePath : SourceControlFileViewColumn::EPathFlags::Default;
 
-	auto CompareTypes = [](const IFileViewTreeItem* Lhs, const IFileViewTreeItem* Rhs)
-	{
-		return FCString::Stricmp(*Lhs->GetType(), *Rhs->GetType());
-	};
-
-	auto CompareLastModified = [](const IFileViewTreeItem* Lhs, const IFileViewTreeItem* Rhs)
-	{
-		const FDateTime& LhsVal = Lhs->GetLastModifiedDateTime();
-		const FDateTime& RhsVal = Rhs->GetLastModifiedDateTime();
-		return LhsVal < RhsVal ? -1 : ( LhsVal == RhsVal ? 0 : 1);
-	};
-
-	auto CompareCheckedOutBy = [](const IFileViewTreeItem* Lhs, const IFileViewTreeItem* Rhs)
-	{
-		return FCString::Stricmp(*Lhs->GetCheckedOutBy(), *Rhs->GetCheckedOutBy());
-	};
-
-	auto GetCompareFunc = [&](const FName& ColumnId)
-	{
-		if (ColumnId == SourceControlFileViewColumn::Icon::Id())
-		{
-			return TFunction<int32(const IFileViewTreeItem*, const IFileViewTreeItem*)>(CompareIcons);
-		}
-		if (ColumnId == SourceControlFileViewColumn::Shelve::Id())
-		{
-			return TFunction<int32(const IFileViewTreeItem*, const IFileViewTreeItem*)>(CompareShelves);
-		}
-		else if (ColumnId == SourceControlFileViewColumn::Name::Id())
-		{
-			return TFunction<int32(const IFileViewTreeItem*, const IFileViewTreeItem*)>(CompareNames);
-		}
-		else if (ColumnId == SourceControlFileViewColumn::Path::Id())
-		{
-			return TFunction<int32(const IFileViewTreeItem*, const IFileViewTreeItem*)>(ComparePaths);
-		}
-		else if (ColumnId == SourceControlFileViewColumn::Type::Id())
-		{
-			return TFunction<int32(const IFileViewTreeItem*, const IFileViewTreeItem*)>(CompareTypes);
-		}
-		else if (ColumnId == SourceControlFileViewColumn::LastModifiedTimestamp::Id())
-		{
-			return TFunction<int32(const IFileViewTreeItem*, const IFileViewTreeItem*)>(CompareLastModified);
-		}
-		else if (ColumnId == SourceControlFileViewColumn::CheckedOutByUser::Id())
-		{
-			return TFunction<int32(const IFileViewTreeItem*, const IFileViewTreeItem*)>(CompareCheckedOutBy);
-		}
-		else
-		{
-			checkNoEntry();
-			return TFunction<int32(const IFileViewTreeItem*, const IFileViewTreeItem*)>();
-		};
-	};
-
-	TFunction<int32(const IFileViewTreeItem*, const IFileViewTreeItem*)> PrimaryCompare = GetCompareFunc(PrimarySortedColumn);
-	TFunction<int32(const IFileViewTreeItem*, const IFileViewTreeItem*)> SecondaryCompare;
+	TFunction<int32(const IFileViewTreeItem&, const IFileViewTreeItem&)> PrimaryCompare = SourceControlFileViewColumn::GetColumnComparer(PrimarySortedColumn, PathFlags);
+	TFunction<int32(const IFileViewTreeItem&, const IFileViewTreeItem&)> SecondaryCompare;
 	if (!SecondarySortedColumn.IsNone())
 	{
-		SecondaryCompare = GetCompareFunc(SecondarySortedColumn);
+		SecondaryCompare = SourceControlFileViewColumn::GetColumnComparer(SecondarySortedColumn, PathFlags);
 	}
 
 	if (PrimarySortMode == EColumnSortMode::Ascending)
@@ -3834,7 +3794,7 @@ void SSourceControlChangelistsWidget::SortFileView()
 		//       thread and can be invoked a lot, the trade off went if favor of speed.
 		FileListNodes.Sort([this, &PrimaryCompare, &SecondaryCompare](const TSharedPtr<IChangelistTreeItem>& Lhs, const TSharedPtr<IChangelistTreeItem>& Rhs)
 		{
-			int32 Result = PrimaryCompare(static_cast<IFileViewTreeItem*>(Lhs.Get()), static_cast<IFileViewTreeItem*>(Rhs.Get()));
+			int32 Result = PrimaryCompare(static_cast<IFileViewTreeItem&>(*Lhs), static_cast<IFileViewTreeItem&>(*Rhs));
 			if (Result < 0)
 			{
 				return true;
@@ -3845,11 +3805,11 @@ void SSourceControlChangelistsWidget::SortFileView()
 			}
 			else if (SecondarySortMode == EColumnSortMode::Ascending)
 			{
-				return SecondaryCompare(static_cast<IFileViewTreeItem*>(Lhs.Get()), static_cast<IFileViewTreeItem*>(Rhs.Get())) < 0;
+				return SecondaryCompare(static_cast<IFileViewTreeItem&>(*Lhs), static_cast<IFileViewTreeItem&>(*Rhs)) < 0;
 			}
 			else
 			{
-				return SecondaryCompare(static_cast<IFileViewTreeItem*>(Lhs.Get()), static_cast<IFileViewTreeItem*>(Rhs.Get())) > 0;
+				return SecondaryCompare(static_cast<IFileViewTreeItem&>(*Lhs), static_cast<IFileViewTreeItem&>(*Rhs)) > 0;
 			}
 		});
 	}
@@ -3857,7 +3817,7 @@ void SSourceControlChangelistsWidget::SortFileView()
 	{
 		FileListNodes.Sort([this, &PrimaryCompare, &SecondaryCompare](const TSharedPtr<IChangelistTreeItem>& Lhs, const TSharedPtr<IChangelistTreeItem>& Rhs)
 		{
-			int32 Result = PrimaryCompare(static_cast<IFileViewTreeItem*>(Lhs.Get()), static_cast<IFileViewTreeItem*>(Rhs.Get()));
+			int32 Result = PrimaryCompare(static_cast<IFileViewTreeItem&>(*Lhs), static_cast<IFileViewTreeItem&>(*Rhs));
 			if (Result > 0)
 			{
 				return true;
@@ -3868,11 +3828,11 @@ void SSourceControlChangelistsWidget::SortFileView()
 			}
 			else if (SecondarySortMode == EColumnSortMode::Ascending)
 			{
-				return SecondaryCompare(static_cast<IFileViewTreeItem*>(Lhs.Get()), static_cast<IFileViewTreeItem*>(Rhs.Get())) < 0;
+				return SecondaryCompare(static_cast<IFileViewTreeItem&>(*Lhs), static_cast<IFileViewTreeItem&>(*Rhs)) < 0;
 			}
 			else
 			{
-				return SecondaryCompare(static_cast<IFileViewTreeItem*>(Lhs.Get()), static_cast<IFileViewTreeItem*>(Rhs.Get())) > 0;
+				return SecondaryCompare(static_cast<IFileViewTreeItem&>(*Lhs), static_cast<IFileViewTreeItem&>(*Rhs)) > 0;
 			}
 		});
 	}
@@ -3914,11 +3874,17 @@ void SSourceControlChangelistsWidget::PopulateItemSearchStrings(const IChangelis
 
 	case IChangelistTreeItem::File:
 	case IChangelistTreeItem::ShelvedFile:
-		SFileTableRow::PopulateSearchString(static_cast<const FFileTreeItem&>(Item), OutStrings);
+		SFileTableRow::PopulateSearchString(
+			static_cast<const FFileTreeItem&>(Item),
+			bShowingContentVersePath ? SourceControlFileViewColumn::EPathFlags::ShowingVersePath : SourceControlFileViewColumn::EPathFlags::Default,
+			OutStrings);
 		break;
 
 	case IChangelistTreeItem::OfflineFile:
-		SOfflineFileTableRow::PopulateSearchString(static_cast<const FOfflineFileTreeItem&>(Item), OutStrings);
+		SOfflineFileTableRow::PopulateSearchString(
+			static_cast<const FOfflineFileTreeItem&>(Item),
+			bShowingContentVersePath ? SourceControlFileViewColumn::EPathFlags::ShowingVersePath : SourceControlFileViewColumn::EPathFlags::Default,
+			OutStrings);
 		break;
 
 	default:
@@ -3936,6 +3902,12 @@ void SSourceControlChangelistsWidget::OnPackageSaved(const FString& PackageFilen
 	OutdatedTimestampFiles.Emplace(PackageFilename); // The background task will convert relative to absolute path.
 }
 
+void SSourceControlChangelistsWidget::OnPluginEdited(IPlugin& InPlugin)
+{
+	// The plugin's Verse path may have changed.  Recompute all Verse paths.
+	bShouldRefreshVersePaths = true;
+}
+
 SListView<FChangelistTreeItemPtr>& SSourceControlChangelistsWidget::GetActiveFileListView() const
 {
 	return *StaticCastSharedPtr<SListView<FChangelistTreeItemPtr>>(FileListViewSwitcher->GetActiveWidget());
diff --git a/Engine/Source/Editor/SourceControlWindows/Private/SSourceControlChangelists.h b/Engine/Source/Editor/SourceControlWindows/Private/SSourceControlChangelists.h
--- a/Engine/Source/Editor/SourceControlWindows/Private/SSourceControlChangelists.h
+++ b/Engine/Source/Editor/SourceControlWindows/Private/SSourceControlChangelists.h
@@ -17,6 +17,7 @@
 #include "Widgets/Views/STreeView.h"
 
 class FChangelistGroupTreeItem;
+class IPlugin;
 class SExpandableArea;
 class SExpandableChangelistArea;
 class SWidgetSwitcher;
@@ -203,6 +204,8 @@ private:
 	/** Invoked when a package is saved, to refresh the last saved timestamp. */
 	void OnPackageSaved(const FString& PackageFilename, UPackage* Package, FObjectPostSaveContext ObjectSaveContext);
 
+	void OnPluginEdited(IPlugin& InPlugin);
+
 private:
 	TSharedPtr<SExpandableChangelistArea> UnsavedAssetsExpandableArea;
 	TSharedPtr<STreeView<FChangelistTreeItemPtr>> UnsavedAssetsTreeView;
@@ -238,6 +241,8 @@ private:
 	int64 UpdateRequestNum = 0;
 	bool bInitialRefreshDone = false;
 	bool bShouldRefresh = false;
+	bool bShouldRefreshVersePaths = false;
+	bool bShowingContentVersePath = false;
 	bool bSourceControlAvailable = false;
 	bool bUpdateMonitoredFileStatusList = false;
 
diff --git a/Engine/Source/Editor/SourceControlWindows/Private/SSourceControlCheckedOutDialog.cpp b/Engine/Source/Editor/SourceControlWindows/Private/SSourceControlCheckedOutDialog.cpp
--- a/Engine/Source/Editor/SourceControlWindows/Private/SSourceControlCheckedOutDialog.cpp
+++ b/Engine/Source/Editor/SourceControlWindows/Private/SSourceControlCheckedOutDialog.cpp
@@ -1,6 +1,7 @@
 // Copyright Epic Games, Inc. All Rights Reserved.
 
 #include "SSourceControlCheckedOutDialog.h"
+#include "AssetToolsModule.h"
 #include "Brushes/SlateDynamicImageBrush.h"
 #include "Styling/AppStyle.h"
 #include "Widgets/SWindow.h"
@@ -13,6 +14,7 @@
 #include "HAL/PlatformFileManager.h"
 #include "Framework/Application/SlateApplication.h"
 #include "SPrimaryButton.h"
+#include "IAssetTools.h"
 #include "ISourceControlModule.h"
 #include "SourceControlAssetDataCache.h"
 #include "Algo/Find.h"
@@ -35,6 +37,7 @@ void SSourceControlCheckedOutDialog::Construct(const FArguments& InArgs)
 	ParentFrame = InArgs._ParentWindow;
 	SortByColumn = SSourceControlConflictWarningWidgetDefs::ColumnID_AssetLabel;
 	SortMode = EColumnSortMode::Ascending;
+	bShowingContentVersePath = FAssetToolsModule::GetModule().Get().ShowingContentVersePath();
 
 	ListViewItems.Reserve(InArgs._Items.Num());
 	for (const auto & Item : InArgs._Items)
@@ -216,7 +219,8 @@ TSharedRef<ITableRow> SSourceControlCheckedOutDialog::OnGenerateRowForList(FFile
 {
 	TSharedRef<ITableRow> Row =
 		SNew(SSourceControlCheckedOutDialogListRow, OwnerTable)
-		.Item(InItem);
+		.Item(InItem)
+		.ShowingContentVersePath(bShowingContentVersePath);
 
 	return Row;
 }
@@ -283,6 +287,7 @@ FReply SSourceControlCheckedOutDialog::OnKeyDown(const FGeometry& MyGeometry, co
 void SSourceControlCheckedOutDialogListRow::Construct(const FArguments& InArgs, const TSharedRef<STableViewBase>& InOwnerTableView)
 {
 	Item = InArgs._Item;
+	bShowingContentVersePath = InArgs._ShowingContentVersePath;
 
 	SMultiColumnTableRow<FFileTreeItemRef>::Construct(FSuperRowType::FArguments(), InOwnerTableView);
 }
@@ -301,7 +306,7 @@ TSharedRef<SWidget> SSourceControlCheckedOutDialogListRow::GenerateWidgetForColu
 			[
 				SNew(STextBlock)
 				.Text(Item->GetAssetName())
-				.ToolTipText(Item->GetPackageName())
+				.ToolTipText(bShowingContentVersePath && Item->GetVersePath().IsValid() ? Item->GetAssetVersePath() : Item->GetAssetPackageName())
 			];
 	}
 	else if (ColumnName == SSourceControlConflictWarningWidgetDefs::ColumnID_AssetClassLabel)
diff --git a/Engine/Source/Editor/SourceControlWindows/Private/SSourceControlCheckedOutDialog.h b/Engine/Source/Editor/SourceControlWindows/Private/SSourceControlCheckedOutDialog.h
--- a/Engine/Source/Editor/SourceControlWindows/Private/SSourceControlCheckedOutDialog.h
+++ b/Engine/Source/Editor/SourceControlWindows/Private/SSourceControlCheckedOutDialog.h
@@ -78,11 +78,6 @@ private:
 	 */
 	void SortTree();
 
-	/**
-	 * Refreshes the assetdata for the items from the AssetDataCache
-	 */
-	void RefreshItemAssetData();
-
 	/** Called when the cancel button is clicked */
 	FReply CloseClicked();
 
@@ -104,7 +99,9 @@ private:
 	FName SortByColumn;
 
 	/** Currently selected sorting mode */
-	EColumnSortMode::Type SortMode;
+	EColumnSortMode::Type SortMode = EColumnSortMode::Ascending;
+
+	bool bShowingContentVersePath = false;
 
 	/** The close button widget. */
 	TSharedPtr<SButton> CloseButton;
@@ -116,9 +113,13 @@ private:
 class SSourceControlCheckedOutDialogListRow : public SMultiColumnTableRow<FFileTreeItemRef>
 {
 public:
-	SLATE_BEGIN_ARGS(SSourceControlCheckedOutDialogListRow) {}
+	SLATE_BEGIN_ARGS(SSourceControlCheckedOutDialogListRow)
+		: _ShowingContentVersePath(false)
+		{
+		}
 
 		SLATE_ARGUMENT(FFileTreeItemPtr, Item)
+		SLATE_ARGUMENT(bool, ShowingContentVersePath)
 
 	SLATE_END_ARGS()
 
@@ -132,4 +133,5 @@ private:
 
 	/** The item associated with this row */
 	FFileTreeItemPtr Item;
+	bool bShowingContentVersePath = false;
 };
\ No newline at end of file
diff --git a/Engine/Source/Editor/SourceControlWindows/Private/SSourceControlCommon.cpp b/Engine/Source/Editor/SourceControlWindows/Private/SSourceControlCommon.cpp
--- a/Engine/Source/Editor/SourceControlWindows/Private/SSourceControlCommon.cpp
+++ b/Engine/Source/Editor/SourceControlWindows/Private/SSourceControlCommon.cpp
@@ -12,6 +12,7 @@
 #include "AssetToolsModule.h"
 #include "Styling/AppStyle.h"
 #include "RevisionControlStyle/RevisionControlStyle.h"
+#include "IAssetTools.h"
 #include "ISourceControlModule.h"
 #include "SourceControlAssetDataCache.h"
 #include "SourceControlHelpers.h"
@@ -66,6 +67,8 @@ void IChangelistTreeItem::RemoveAllChildren()
 	Children.Reset();
 }
 
+namespace SSourceControlCommonPrivate
+{
 static FString RetrieveAssetName(const FAssetData& InAssetData)
 {
 	static const FName NAME_ActorLabel(TEXT("ActorLabel"));
@@ -122,41 +125,74 @@ static FString RetrieveAssetTypeName(const FAssetData& InAssetData)
 	return InAssetData.AssetClassPath.ToString();
 }
 
-static void RefreshAssetInformationInternal(const TArray<FAssetData>& Assets, const FString& InFilename, FString& OutAssetName, FString& OutAssetPath, FString& OutAssetType, FString& OutAssetTypeName, FText& OutPackageName, FColor& OutAssetTypeColor)
+static const FAssetData* GetUserFacingAsset(const TArray<FAssetData>* Assets, int32& OutNumUserFacingAssets)
+{
+	if (Assets == nullptr || Assets->IsEmpty())
+	{
+		OutNumUserFacingAssets = 0;
+		return nullptr;
+	}
+
+	auto IsNotRedirector = [](const FAssetData& InAssetData) { return !InAssetData.IsRedirector(); };
+	OutNumUserFacingAssets = Algo::CountIf(*Assets, IsNotRedirector);
+
+	if (OutNumUserFacingAssets == 1)
+	{
+		return Algo::FindByPredicate(*Assets, IsNotRedirector);
+	}
+	else
+	{
+		return Assets->GetData();
+	}
+}
+
+static bool RefreshAssetVersePathInternal(const TArray<FAssetData>* Assets, UE::Core::FVersePath& InOutAssetVersePath)
+{
+	UE::Core::FVersePath TempAssetVersePath;
+
+	{
+		int32 NumUserFacingAsset = 0;
+		const FAssetData* AssetData = GetUserFacingAsset(Assets, NumUserFacingAsset);
+		if (AssetData != nullptr && NumUserFacingAsset == 1)
+		{
+			TempAssetVersePath = AssetData->GetVersePath();
+		}
+	}
+
+	if (TempAssetVersePath != InOutAssetVersePath)
+	{
+		InOutAssetVersePath = MoveTemp(TempAssetVersePath);
+		return true;
+	}
+
+	return false;
+}
+
+static void RefreshAssetInformationInternal(const TArray<FAssetData>* Assets, const FString& InFilename, FString& OutAssetName, FString& OutAssetPath, UE::Core::FVersePath& OutAssetVersePath, FString& OutAssetType, FString& OutAssetTypeName, FString& OutPackageName, FColor& OutAssetTypeColor)
 {
 	// Initialize display-related members
-	FString Filename = InFilename;
-	FString Extension = FPaths::GetExtension(Filename);
 	FString TempAssetName = SSourceControlCommon::GetDefaultAssetName().ToString();
-	FString TempAssetPath = Filename;
+	FString TempAssetPath;
+	UE::Core::FVersePath TempAssetVersePath;
 	FString TempAssetType = SSourceControlCommon::GetDefaultAssetType().ToString();
 	FString TempAssetTypeName = SSourceControlCommon::GetDefaultAssetType().ToString();
-	FString TempPackageName = Filename;
+	FString TempPackageName;
 	FColor TempAssetColor = FColor(		// Copied from ContentBrowserCLR.cpp
 		127 + FColor::Red.R / 2,	// Desaturate the colors a bit (GB colors were too.. much)
 		127 + FColor::Red.G / 2,
 		127 + FColor::Red.B / 2,
 		200); // Opacity
 
-
+	const FString Extension = FPaths::GetExtension(InFilename);
 	bool bIsPackageExtension = 
 		FPackageName::IsPackageExtension(*Extension) ||
 		FPackageName::IsVerseExtension(*Extension);
 
-	if (Assets.Num() > 0)
+	int32 NumUserFacingAsset = 0;
+	const FAssetData* AssetData = GetUserFacingAsset(Assets, NumUserFacingAsset);
+	if (AssetData != nullptr)
 	{
-		auto IsNotRedirector = [](const FAssetData& InAssetData) { return !InAssetData.IsRedirector(); };
-		const int32 NumUserFacingAsset = Algo::CountIf(Assets, IsNotRedirector);
-		const FAssetData* AssetData = nullptr;
-
-		if (NumUserFacingAsset == 1)
-		{
-			AssetData = Algo::FindByPredicate(Assets, IsNotRedirector);
-		}
-		else
-		{
-			AssetData = &Assets[0];
-		}
+		FAssetToolsModule& AssetToolsModule = FAssetToolsModule::GetModule();
 
 		TempAssetName = RetrieveAssetName(*AssetData);
 		TempAssetPath = RetrieveAssetPath(*AssetData);
@@ -167,16 +203,23 @@ static void RefreshAssetInformationInternal(const TArray<FAssetData>& Assets, co
 		{
 			TempAssetType = SSourceControlCommon::GetDefaultMultipleAsset().ToString();
 
-			for (int32 i = 1; i < Assets.Num(); ++i)
+			for (const FAssetData& OtherAssetData : *Assets)
 			{
-				TempAssetName += TEXT(";") + RetrieveAssetName(Assets[i]);
+				if (AssetData != &OtherAssetData)
+				{
+					TempAssetName += TEXT(";") + RetrieveAssetName(OtherAssetData);
+				}
 			}			
 		}
 		else
 		{
 			TempAssetType = AssetData->AssetClassPath.ToString();
 
-			const FAssetToolsModule& AssetToolsModule = FModuleManager::LoadModuleChecked<FAssetToolsModule>(TEXT("AssetTools"));
+			if (AssetToolsModule.Get().GetOnShowingContentVersePath().IsBound())
+			{
+				TempAssetVersePath = AssetData->GetVersePath();
+			}
+
 			const TSharedPtr<IAssetTypeActions> AssetTypeActions = AssetToolsModule.Get().GetAssetTypeActionsForClass(AssetData->GetClass()).Pin();
 
 			if (AssetTypeActions.IsValid())
@@ -188,7 +231,7 @@ static void RefreshAssetInformationInternal(const TArray<FAssetData>& Assets, co
 		// Beautify the package name
 		TempPackageName = TempAssetPath + "." + TempAssetName;
 	}
-	else if (bIsPackageExtension && FPackageName::TryConvertFilenameToLongPackageName(Filename, TempPackageName))
+	else if (bIsPackageExtension && FPackageName::TryConvertFilenameToLongPackageName(InFilename, TempPackageName))
 	{
 		// Fake asset name, asset path from the package name
 		TempAssetPath = TempPackageName;
@@ -202,9 +245,10 @@ static void RefreshAssetInformationInternal(const TArray<FAssetData>& Assets, co
 	}
 	else
 	{
-		TempAssetName = FPaths::GetCleanFilename(Filename);
-		TempPackageName = Filename; // Put back original package name if the try failed
-		TempAssetType = FText::Format(SSourceControlCommon::GetDefaultUnknownAssetType(), FText::FromString(FPaths::GetExtension(Filename).ToUpper())).ToString();
+		TempAssetName = FPaths::GetCleanFilename(InFilename);
+		TempAssetPath = InFilename;
+		TempPackageName = InFilename; // Put back original package name if the try failed
+		TempAssetType = FText::Format(SSourceControlCommon::GetDefaultUnknownAssetType(), FText::FromString(Extension.ToUpper())).ToString();
 		TempAssetTypeName = TempAssetType;
 
 		// Attempt to make package name relative to one of the project roots instead of a full absolute path
@@ -218,21 +262,25 @@ static void RefreshAssetInformationInternal(const TArray<FAssetData>& Assets, co
 				break;
 			}
 		}
+
+		FPaths::MakePlatformFilename(TempAssetPath);
+		FPaths::MakePlatformFilename(TempPackageName);
 	}
 
 	// Finally, assign the temp variables to the member variables
-	OutAssetName = TempAssetName;
-	OutAssetPath = TempAssetPath;
-	OutAssetType = TempAssetType;
-	OutAssetTypeName = TempAssetTypeName;
+	OutAssetName = MoveTemp(TempAssetName);
+	OutAssetPath = MoveTemp(TempAssetPath);
+	OutAssetVersePath = MoveTemp(TempAssetVersePath);
+	OutAssetType = MoveTemp(TempAssetType);
+	OutAssetTypeName = MoveTemp(TempAssetTypeName);
 	OutAssetTypeColor = TempAssetColor;
-	OutPackageName = FText::FromString(TempPackageName);
+	OutPackageName = MoveTemp(TempPackageName);
+}
 }
 
 //////////////////////////////////////////////////////////////////////////
 
-FString IFileViewTreeItem::DefaultStrValue; // Default is an empty string.
-FDateTime IFileViewTreeItem::DefaultDateTimeValue; // Default is FDateTime::MinValue().
+const FString IFileViewTreeItem::DefaultStrValue; // Default is an empty string.
 
 void IFileViewTreeItem::SetLastModifiedDateTime(const FDateTime& Timestamp)
 {
@@ -260,12 +308,11 @@ FString FUnsavedAssetsTreeItem::GetDisplayString() const
 FFileTreeItem::FFileTreeItem(FSourceControlStateRef InFileState, bool bBeautifyPaths, bool bIsShelvedFile)
 	: IFileViewTreeItem(bIsShelvedFile ? IChangelistTreeItem::ShelvedFile : IChangelistTreeItem::File)
 	, FileState(InFileState)
+	, CheckBoxState(ECheckBoxState::Checked)
 	, MinTimeBetweenUpdate(FTimespan::FromSeconds(5.f))
 	, LastUpdateTime()
 	, bAssetsUpToDate(false)
 {
-	CheckBoxState = ECheckBoxState::Checked;
-
 	// Initialize asset data first
 
 	if (bBeautifyPaths)
@@ -306,15 +353,28 @@ FText FFileTreeItem::GetCheckedOutByUser() const
 	return FText::FromString(GetCheckedOutBy());
 }
 
+FText FFileTreeItem::GetFileName() const
+{
+	FString Filename = FileState->GetFilename();
+	FPaths::MakePlatformFilename(Filename);
+	return FText::FromString(MoveTemp(Filename));
+}
+
 void FFileTreeItem::RefreshAssetInformation()
 {
 	// Initialize display-related members
-	static TArray<FAssetData> NoAssets;
-	RefreshAssetInformationInternal(Assets.IsValid() ? *Assets : NoAssets, FileState->GetFilename(), AssetNameStr, AssetPathStr, AssetTypeStr, AssetTypeNameStr, PackageName, AssetTypeColor);
+	SSourceControlCommonPrivate::RefreshAssetInformationInternal(Assets.Get(), FileState->GetFilename(), AssetNameStr, AssetPathStr, AssetVersePathStruct, AssetTypeStr, AssetTypeNameStr, AssetPackageNameStr, AssetTypeColor);
 	AssetName = FText::FromString(AssetNameStr);
 	AssetPath = FText::FromString(AssetPathStr);
+	AssetVersePath = FText::FromString(AssetVersePathStruct.ToString());
 	AssetType = FText::FromString(AssetTypeStr);
 	AssetTypeName = FText::FromString(AssetTypeNameStr);
+	AssetPackageName = FText::FromString(AssetPackageNameStr);
+}
+
+bool FFileTreeItem::RefreshVersePath()
+{
+	return SSourceControlCommonPrivate::RefreshAssetVersePathInternal(Assets.Get(), AssetVersePathStruct);
 }
 
 FText FFileTreeItem::GetAssetName() const
@@ -352,16 +412,9 @@ FText FShelvedChangelistTreeItem::GetDisplayText() const
 
 FOfflineFileTreeItem::FOfflineFileTreeItem(const FString& InFilename)
 	: IFileViewTreeItem(IChangelistTreeItem::OfflineFile)
-	, Assets()
+	, CheckBoxState(ECheckBoxState::Checked)
 	, Filename(InFilename)
-	, PackageName(FText::FromString(InFilename)) 
-	, AssetName(SSourceControlCommon::GetDefaultAssetName())
-	, AssetPath()
-	, AssetType(SSourceControlCommon::GetDefaultAssetType())
-	, AssetTypeColor()
 {
-	FString TempString;
-
 	USourceControlHelpers::GetAssetData(InFilename, Assets);
 
 	RefreshAssetInformation();
@@ -369,11 +422,18 @@ FOfflineFileTreeItem::FOfflineFileTreeItem(const FString& InFilename)
 
 void FOfflineFileTreeItem::RefreshAssetInformation()
 {
-	RefreshAssetInformationInternal(Assets, Filename, AssetNameStr, AssetPathStr, AssetTypeStr, AssetTypeNameStr, PackageName, AssetTypeColor);
+	SSourceControlCommonPrivate::RefreshAssetInformationInternal(&Assets, Filename, AssetNameStr, AssetPathStr, AssetVersePathStruct, AssetTypeStr, AssetTypeNameStr, AssetPackageNameStr, AssetTypeColor);
 	AssetName = FText::FromString(AssetNameStr);
 	AssetPath = FText::FromString(AssetPathStr);
+	AssetVersePath = FText::FromString(AssetVersePathStruct.ToString());
 	AssetType = FText::FromString(AssetTypeStr);
 	AssetTypeName = FText::FromString(AssetTypeNameStr);
+	AssetPackageName = FText::FromString(AssetPackageNameStr);
+}
+
+bool FOfflineFileTreeItem::RefreshVersePath()
+{
+	return SSourceControlCommonPrivate::RefreshAssetVersePathInternal(&Assets, AssetVersePathStruct);
 }
 
 //////////////////////////////////////////////////////////////////////////
diff --git a/Engine/Source/Editor/SourceControlWindows/Private/SSourceControlCommon.h b/Engine/Source/Editor/SourceControlWindows/Private/SSourceControlCommon.h
--- a/Engine/Source/Editor/SourceControlWindows/Private/SSourceControlCommon.h
+++ b/Engine/Source/Editor/SourceControlWindows/Private/SSourceControlCommon.h
@@ -3,6 +3,7 @@
 #pragma once
 
 #include "CoreMinimal.h"
+#include "Containers/VersePath.h"
 #include "Styling/SlateColor.h"
 #include "Styling/SlateTypes.h"
 #include "SourceControlAssetDataCache.h"
@@ -91,17 +92,29 @@ protected:
  */
 struct IFileViewTreeItem : public IChangelistTreeItem
 {
+	/** Returns the checkbox state of this item */
+	virtual ECheckBoxState GetCheckBoxState() const = 0;
+
+	/** Sets the checkbox state of this item */
+	virtual void SetCheckBoxState(ECheckBoxState NewState) = 0;
+
 	/**The 'Priority' given to the item icon when sorting ascending (lower will be sorted first). */
-	virtual int32 GetIconSortingPriority() const { return 0; }
+	virtual int32 GetIconSortingPriority() const = 0;
 
 	/** The values displayed in the 'Name' column. */
-	virtual const FString& GetName() const { return DefaultStrValue; }
+	virtual const FString& GetName() const = 0;
+
+	/** The values displayed in the 'Path' column when not displaying package names or Verse paths. */
+	virtual const FString& GetPath() const = 0;
+
+	/** The values displayed in the 'Path' column when displaying package names. */
+	virtual const FString& GetPackageName() const = 0;
 
-	/** The values displayed in the 'Path' column. */
-	virtual const FString& GetPath() const { return DefaultStrValue; }
+	/** The values displayed in the 'Path' column when displaying Verse paths. */
+	virtual const UE::Core::FVersePath& GetVersePath() const = 0;
 
 	/** The values displayed in the 'Type' column. */
-	virtual const FString& GetType() const { return DefaultStrValue; }
+	virtual const FString& GetType() const = 0;
 
 	/** The values displayed in the 'Type Name' column. */
 	virtual const FString& GetTypeName() const { return DefaultStrValue; }
@@ -110,7 +123,7 @@ struct IFileViewTreeItem : public IChangelistTreeItem
 	virtual const FString& GetCheckedOutBy() const { return DefaultStrValue; }
 
 	/** Returns the full pathname of the files on the file sytem. */
-	virtual const FString& GetFullPathname() const { return DefaultStrValue; }
+	virtual const FString& GetFullPathname() const = 0;
 
 	/** Set the last modified time timestamp. */
 	void SetLastModifiedDateTime(const FDateTime& Timestamp);
@@ -121,6 +134,8 @@ struct IFileViewTreeItem : public IChangelistTreeItem
 	/** The values displayed in the 'Last Modified' column as text. */
 	FText GetLastModifiedTimestamp() const { return LastModifiedTimestampText; }
 
+	virtual bool RefreshVersePath() = 0;
+
 public:
 	/** Keep the icon sorting priority as it was the last time the item was displayed. Used to detect if the priority changed between two refreshes of the UI. */
 	int32 DisplayedIconPriority = -1;
@@ -130,8 +145,7 @@ protected:
 
 private:
 	// Use an empty string to return as default for const FString&.
-	static FString DefaultStrValue;
-	static FDateTime DefaultDateTimeValue;
+	static const FString DefaultStrValue;
 
 	/** The timestamp of the last modification to the file. */
 	FText LastModifiedTimestampText;
@@ -224,14 +238,20 @@ struct FFileTreeItem : public IFileViewTreeItem
 {
 	explicit FFileTreeItem(FSourceControlStateRef InFileState, bool bBeautifyPaths = true, bool bIsShelvedFile = false);
 
+	virtual ECheckBoxState GetCheckBoxState() const override { return CheckBoxState; }
+	virtual void SetCheckBoxState(ECheckBoxState NewState) override { CheckBoxState = NewState; }
 	virtual int32 GetIconSortingPriority() const override;
 	virtual const FString& GetName() const override { return AssetNameStr; }
 	virtual const FString& GetPath() const override { return AssetPathStr; }
+	virtual const FString& GetPackageName() const override { return AssetPackageNameStr; }
+	virtual const UE::Core::FVersePath& GetVersePath() const override { return AssetVersePathStruct; }
 	virtual const FString& GetType() const override { return AssetTypeStr; }
 	virtual const FString& GetTypeName() const override { return AssetTypeNameStr; }
 	virtual const FString& GetFullPathname() const override { return FileState->GetFilename(); }
 	virtual const FString& GetCheckedOutBy() const override;
 
+	virtual bool RefreshVersePath() override;
+
 	/** Updates informations based on AssetData */
 	void RefreshAssetInformation();
 
@@ -244,6 +264,9 @@ struct FFileTreeItem : public IFileViewTreeItem
 	/** Returns the asset path of the item */
 	FText GetAssetPath() const { return AssetPath; }
 
+	/** Returns the asset Verse path of the item */
+	FText GetAssetVersePath() const { return AssetVersePath; }
+
 	/** Returns the asset type of the item */
 	FText GetAssetType() const { return AssetType; }
 
@@ -257,22 +280,10 @@ struct FFileTreeItem : public IFileViewTreeItem
 	FText GetCheckedOutByUser() const;
 
 	/** Returns the package name of the item to display */
-	FText GetPackageName() const { return PackageName; }
+	FText GetAssetPackageName() const { return AssetPackageName; }
 
 	/** Returns the file name of the item in source control */
-	FText GetFileName() const { return FText::FromString(FileState->GetFilename()); }
-
-	/** Returns the name of the icon to be used in the list item widget */
-	FName GetIconName() const { return FileState->GetIcon().GetStyleName(); }
-
-	/** Returns the tooltip text for the icon */
-	FText GetIconTooltip() const { return FileState->GetDisplayTooltip(); }
-
-	/** Returns the checkbox state of this item */
-	ECheckBoxState GetCheckBoxState() const { return CheckBoxState; }
-
-	/** Sets the checkbox state of this item */
-	void SetCheckBoxState(ECheckBoxState NewState) { CheckBoxState = NewState; }
+	FText GetFileName() const;
 
 	/** true if the item is not in source control and needs to be added prior to checkin */
 	bool NeedsAdding() const { return !FileState->IsSourceControlled(); }
@@ -301,7 +312,7 @@ public:
 	FSourceControlStateRef FileState;
 
 private:
-	/** Checkbox state, used only in the Submit dialog */
+	/** Checkbox state */
 	ECheckBoxState CheckBoxState;
 
 	/** Cached asset name to display */
@@ -312,6 +323,14 @@ private:
 	FText AssetPath;
 	FString AssetPathStr;
 
+	/** Cached package name to display */
+	FText AssetPackageName;
+	FString AssetPackageNameStr;
+
+	/** Cached Verse path to display */
+	FText AssetVersePath;
+	UE::Core::FVersePath AssetVersePathStruct;
+
 	/** Cached asset type to display */
 	FText AssetType;
 	FString AssetTypeStr;
@@ -323,9 +342,6 @@ private:
 	/** Cached asset type related color to display */
 	FColor AssetTypeColor;
 
-	/** Cached package name to display */
-	FText PackageName;
-
 	/** The other user that has the checked out. */
 	mutable FString CheckedOutBy;
 
@@ -359,26 +375,38 @@ struct FOfflineFileTreeItem : public IFileViewTreeItem
 	void RefreshAssetInformation();
 
 public:
+	virtual ECheckBoxState GetCheckBoxState() const override { return CheckBoxState; }
+	virtual void SetCheckBoxState(ECheckBoxState NewState) override { CheckBoxState = NewState; }
+	virtual int32 GetIconSortingPriority() const override { return 0; }
 	virtual const FString& GetName() const override { return AssetNameStr; }
 	virtual const FString& GetPath() const override { return AssetPathStr; }
+	virtual const FString& GetPackageName() const override { return AssetPackageNameStr; }
+	virtual const UE::Core::FVersePath& GetVersePath() const override { return AssetVersePathStruct; }
 	virtual const FString& GetType() const override { return AssetTypeStr; }
 	virtual const FString& GetFullPathname() const override { return Filename; }
 
+	virtual bool RefreshVersePath() override;
+
 	const FString& GetFilename() const { return Filename; }
-	const FText& GetPackageName() const { return PackageName; }
+	const FText& GetDisplayPackageName() const { return AssetPackageName; }
 	const FText& GetDisplayName() const { return AssetName; }
 	const FText& GetDisplayPath() const { return AssetPath; }
+	const FText& GetDisplayVersePath() const { return AssetVersePath; }
 	const FText& GetDisplayType() const { return AssetType; }
 	const FColor& GetDisplayColor() const { return AssetTypeColor; }
 
 private:
+	ECheckBoxState CheckBoxState;
 	TArray<FAssetData> Assets;
 	FString Filename;
-	FText PackageName;
 	FText AssetName;
 	FString AssetNameStr;
 	FText AssetPath;
 	FString AssetPathStr;
+	FText AssetPackageName;
+	FString AssetPackageNameStr;
+	FText AssetVersePath;
+	UE::Core::FVersePath AssetVersePathStruct;
 	FText AssetType;
 	FString AssetTypeStr;
 	FText AssetTypeName;
diff --git a/Engine/Source/Editor/SourceControlWindows/Private/SSourceControlFileDialog.cpp b/Engine/Source/Editor/SourceControlWindows/Private/SSourceControlFileDialog.cpp
--- a/Engine/Source/Editor/SourceControlWindows/Private/SSourceControlFileDialog.cpp
+++ b/Engine/Source/Editor/SourceControlWindows/Private/SSourceControlFileDialog.cpp
@@ -2,7 +2,10 @@
 
 #include "SSourceControlFileDialog.h"
 
+#include "AssetToolsModule.h"
 #include "Framework/Application/SlateApplication.h"
+#include "IAssetTools.h"
+#include "RevisionControlStyle/RevisionControlStyle.h"
 #include "SSourceControlChangelistRows.h"
 #include "Styling/AppStyle.h"
 #include "SWarningOrErrorBox.h"
@@ -16,6 +19,10 @@
 
 void SSourceControlFileDialog::Construct(const FArguments& InArgs)
 {
+	SortByColumn = SourceControlFileViewColumn::Name::Id();
+
+	bShowingContentVersePath = FAssetToolsModule::GetModule().Get().ShowingContentVersePath();
+
 	Reset();
 	SetMessage(InArgs._Message);
 	SetWarning(InArgs._Warning);
@@ -28,21 +35,51 @@ void SSourceControlFileDialog::Construct(const FArguments& InArgs)
 					.HeaderRow
 					(
 						SNew(SHeaderRow)
-						+ SHeaderRow::Column("Icon")
-						.DefaultLabel(FText::GetEmpty())
+						+ SHeaderRow::Column(SourceControlFileViewColumn::Icon::Id())
+						.DefaultTooltip(SourceControlFileViewColumn::Icon::GetToolTipText())
 						.FillSized(18)
+						.HeaderContentPadding(FMargin(0.0f))
+						.SortMode(this, &SSourceControlFileDialog::GetColumnSortMode, SourceControlFileViewColumn::Icon::Id())
+						.OnSort(this, &SSourceControlFileDialog::OnColumnSortModeChanged)
+						[
+							SNew(SHorizontalBox)
+							+ SHorizontalBox::Slot()
+							.Padding(1.0f, 0.0f)
+							[
+								SNew(SBox)
+								.WidthOverride(16.0f)
+								.HeightOverride(16.0f)
+								.HAlign(HAlign_Center)
+								.VAlign(VAlign_Center)
+								.Visibility(this, &SSourceControlFileDialog::GetIconColumnContentVisibility)
+								[
+									SNew(SImage)
+									.ColorAndOpacity(FSlateColor::UseSubduedForeground())
+									.Image(FRevisionControlStyleManager::Get().GetBrush("RevisionControl.ChangelistsTab"))
+								]
+							]
+						]
 
-						+ SHeaderRow::Column("Name")
-						.DefaultLabel(LOCTEXT("Name", "Name"))
+						+ SHeaderRow::Column(SourceControlFileViewColumn::Name::Id())
+						.DefaultLabel(SourceControlFileViewColumn::Name::GetDisplayText())
+						.DefaultTooltip(SourceControlFileViewColumn::Name::GetToolTipText())
 						.FillWidth(0.2f)
+						.SortMode(this, &SSourceControlFileDialog::GetColumnSortMode, SourceControlFileViewColumn::Name::Id())
+						.OnSort(this, &SSourceControlFileDialog::OnColumnSortModeChanged)
 
-						+ SHeaderRow::Column("Path")
-						.DefaultLabel(LOCTEXT("Path", "Path"))
+						+ SHeaderRow::Column(SourceControlFileViewColumn::Path::Id())
+						.DefaultLabel(SourceControlFileViewColumn::Path::GetDisplayText())
+						.DefaultTooltip(SourceControlFileViewColumn::Path::GetToolTipText())
 						.FillWidth(0.6f)
+						.SortMode(this, &SSourceControlFileDialog::GetColumnSortMode, SourceControlFileViewColumn::Path::Id())
+						.OnSort(this, &SSourceControlFileDialog::OnColumnSortModeChanged)
 
-						+ SHeaderRow::Column("Type")
-						.DefaultLabel(LOCTEXT("Type", "Type"))
+						+ SHeaderRow::Column(SourceControlFileViewColumn::Type::Id())
+						.DefaultLabel(SourceControlFileViewColumn::Type::GetDisplayText())
+						.DefaultTooltip(SourceControlFileViewColumn::Type::GetToolTipText())
 						.FillWidth(0.2f)
+						.SortMode(this, &SSourceControlFileDialog::GetColumnSortMode, SourceControlFileViewColumn::Type::Id())
+						.OnSort(this, &SSourceControlFileDialog::OnColumnSortModeChanged)
 					);
 
 	ChildSlot
@@ -129,16 +166,60 @@ void SSourceControlFileDialog::SetFiles(const TArray<FSourceControlStateRef>& In
 		FileTreeNodes.Add(MakeShared<FFileTreeItem>(FileState, true));
 	}
 
+	SortFiles();
+
 	if (FileTreeView.IsValid())
 	{
 		FileTreeView->RequestTreeRefresh();
 	}
 }
 
+EColumnSortMode::Type SSourceControlFileDialog::GetColumnSortMode(FName ColumnId) const
+{
+	if (SortByColumn != ColumnId)
+	{
+		return EColumnSortMode::None;
+	}
+
+	return SortMode;
+}
+
+void SSourceControlFileDialog::OnColumnSortModeChanged(EColumnSortPriority::Type SortPriority, const FName& ColumnId, EColumnSortMode::Type InSortMode)
+{
+	SortByColumn = ColumnId;
+	SortMode = InSortMode;
+
+	SortFiles();
+	FileTreeView->RequestListRefresh();
+}
+
+EVisibility SSourceControlFileDialog::GetIconColumnContentVisibility() const
+{
+	// Hide the icon when sorting the icon column (it clashes with the sort mode icon).
+	return GetColumnSortMode(SourceControlFileViewColumn::Icon::Id()) == EColumnSortMode::None ? EVisibility::Visible : EVisibility::Collapsed;
+}
+
+void SSourceControlFileDialog::SortFiles()
+{
+	TFunction<bool(const IFileViewTreeItem&, const IFileViewTreeItem&)> SortPredicate = SourceControlFileViewColumn::GetSortPredicate(
+		SortMode, SortByColumn, bShowingContentVersePath ? SourceControlFileViewColumn::EPathFlags::ShowingVersePath : SourceControlFileViewColumn::EPathFlags::Default);
+	if (SortPredicate)
+	{
+		Algo::SortBy(
+			FileTreeNodes,
+			[](const FChangelistTreeItemPtr& ListViewItem) -> const IFileViewTreeItem&
+			{
+				return static_cast<IFileViewTreeItem&>(*ListViewItem);
+			},
+			SortPredicate);
+	}
+}
+
 TSharedRef<ITableRow> SSourceControlFileDialog::OnGenerateRow(TSharedPtr<IChangelistTreeItem> InItem, const TSharedRef<STableViewBase>& OwnerTable)
 {
 	return	SNew(SFileTableRow, OwnerTable)
-			.TreeItemToVisualize(InItem);
+			.TreeItemToVisualize(InItem)
+			.PathFlags(bShowingContentVersePath ? SourceControlFileViewColumn::EPathFlags::ShowingVersePath : SourceControlFileViewColumn::EPathFlags::Default);
 }
 
 void SSourceControlFileDialog::Reset()
diff --git a/Engine/Source/Editor/SourceControlWindows/Private/SSourceControlFileDialog.h b/Engine/Source/Editor/SourceControlWindows/Private/SSourceControlFileDialog.h
--- a/Engine/Source/Editor/SourceControlWindows/Private/SSourceControlFileDialog.h
+++ b/Engine/Source/Editor/SourceControlWindows/Private/SSourceControlFileDialog.h
@@ -76,6 +76,14 @@ private:
 	virtual FReply OnKeyDown(const FGeometry& MyGeometry, const FKeyEvent& InKeyEvent) override;
 	//~ End SWidget Interface
 
+	EColumnSortMode::Type GetColumnSortMode(FName ColumnId) const;
+
+	void OnColumnSortModeChanged(EColumnSortPriority::Type SortPriority, const FName& ColumnId, EColumnSortMode::Type InSortMode);
+
+	EVisibility GetIconColumnContentVisibility() const;
+
+	void SortFiles();
+
 	/** Treeview callback */
 	TSharedRef<ITableRow> OnGenerateRow(TSharedPtr<IChangelistTreeItem> InItem, const TSharedRef<STableViewBase>& OwnerTable);
 
@@ -104,6 +112,9 @@ private:
 	/** The Window owning the dialog */
 	TSharedPtr<SWindow> Window;
 
+	FName SortByColumn;
+	EColumnSortMode::Type SortMode = EColumnSortMode::Ascending;
+
 	/** Hold the nodes displayed by the file tree. */
 	TArray<TSharedPtr<IChangelistTreeItem>> FileTreeNodes;
 
@@ -122,6 +133,8 @@ private:
 	/** The warning to display */
 	FText Warning;
 
+	bool bShowingContentVersePath = false;
+
 	/** Set to true if the user clicked on the primary button */
-	bool bIsProceedButtonPressed;
+	bool bIsProceedButtonPressed = false;
 };
diff --git a/Engine/Source/Editor/SourceControlWindows/Private/SSourceControlHistory.cpp b/Engine/Source/Editor/SourceControlWindows/Private/SSourceControlHistory.cpp
--- a/Engine/Source/Editor/SourceControlWindows/Private/SSourceControlHistory.cpp
+++ b/Engine/Source/Editor/SourceControlWindows/Private/SSourceControlHistory.cpp
@@ -1045,6 +1045,8 @@ private:
 		if (TreeItemPtr->FileListItem.IsValid())
 		{
 			TSharedPtr<FHistoryFileListViewItem> FileListItem = TreeItemPtr->FileListItem;
+			FString FileName = FileListItem->FileName;
+			FPaths::MakePlatformFilename(FileName);
 
 			return
 				SNew(STableRow< TSharedPtr<FName> >, OwnerTable)
@@ -1056,7 +1058,7 @@ private:
 					[
 						SNew( STextBlock )
 						.Font( FAppStyle::GetFontStyle( TEXT("BoldFont") ))
-						.Text( FText::FromString(FileListItem->FileName) )
+						.Text( FText::FromString(MoveTemp(FileName)) )
 					]
 				]
 				.OnDragDetected(this, &SSourceControlHistoryWidget::OnRowDragDetected)
diff --git a/Engine/Source/Editor/SourceControlWindows/Private/SSourceControlRevert.cpp b/Engine/Source/Editor/SourceControlWindows/Private/SSourceControlRevert.cpp
--- a/Engine/Source/Editor/SourceControlWindows/Private/SSourceControlRevert.cpp
+++ b/Engine/Source/Editor/SourceControlWindows/Private/SSourceControlRevert.cpp
@@ -1,11 +1,15 @@
 // Copyright Epic Games, Inc. All Rights Reserved.
 
 #include "CoreMinimal.h"
+#include "AssetToolsModule.h"
 #include "Misc/MessageDialog.h"
+#include "IAssetTools.h"
 #include "ISourceControlOperation.h"
+#include "RevisionControlStyle/RevisionControlStyle.h"
 #include "SourceControlOperations.h"
 #include "SourceControlWindows.h"
 #include "SourceControlHelpers.h"
+#include "SSourceControlChangelistRows.h"
 #include "ISourceControlModule.h"
 #include "UObject/Package.h"
 #include "Misc/PackageName.h"
@@ -21,7 +25,6 @@
 #include "Widgets/Layout/SBorder.h"
 #include "Widgets/Images/SImage.h"
 #include "Widgets/Text/STextBlock.h"
-#include "Widgets/Layout/SUniformGridPanel.h"
 #include "Widgets/Input/SButton.h"
 #include "Widgets/Views/STableViewBase.h"
 #include "Widgets/Views/STableRow.h"
@@ -40,48 +43,79 @@
 //-------------------------------------
 //Source Control Window Constants
 //-------------------------------------
-namespace ERevertResults
+enum class ERevertResults
 {
-	enum Type
-	{
-		REVERT_ACCEPTED,
-		REVERT_CANCELED
-	};
-}
-
+	Accepted,
+	Canceled
+};
 
-struct FRevertCheckBoxListViewItem
+/** Additional state for each IFileViewTreeItem row. */
+struct FRevertTableRowState
 {
-	/**
-	 * Constructor
-	 *
-	 * @param	InText	String that should appear for the item in the list view
-	 */
-	FRevertCheckBoxListViewItem( FString InText )
+	FRevertTableRowState(const FString& InPackageName, bool bInIsModified)
+		: bIsModified(bInIsModified)
+		, PackageName(InPackageName)
 	{
-		Text = InText;
-		IsSelected = true;
-		IsModified = false;
 	}
 
-	void OnCheckStateChanged( const ECheckBoxState NewCheckedState )
-	{
-		IsSelected = (NewCheckedState == ECheckBoxState::Checked);
-	}
+	const bool bIsModified;
+	const FString PackageName;
+};
 
-	ECheckBoxState OnIsChecked() const
+/** Row widget - derives from either SFileTableRow or SOfflineFileTableRow depending on whether we are performing a unsaved revert or not. */
+template<typename CommonTableRowType>
+class SSourceControlRevertRow : public CommonTableRowType
+{
+public:
+	SLATE_BEGIN_ARGS(SSourceControlRevertRow)
+		: _ShowingContentVersePath(false)
+		{
+		}
+
+		SLATE_ARGUMENT(FChangelistTreeItemPtr, TreeItemToVisualize)
+		SLATE_ARGUMENT(bool, IsModified)
+		SLATE_ARGUMENT(bool, ShowingContentVersePath)
+	SLATE_END_ARGS()
+
+public:
+	void Construct(const FArguments& InArgs, const TSharedRef<STableViewBase>& InOwner)
 	{
-		return ( IsSelected ) ? ECheckBoxState::Checked : ECheckBoxState::Unchecked;
+		bIsModified = InArgs._IsModified;
+
+		CommonTableRowType::Construct(
+			typename CommonTableRowType::FArguments()
+			.TreeItemToVisualize(InArgs._TreeItemToVisualize)
+			.PathFlags(InArgs._ShowingContentVersePath ? SourceControlFileViewColumn::EPathFlags::ShowingVersePath : SourceControlFileViewColumn::EPathFlags::Default),
+			InOwner);
 	}
 
-	EVisibility OnGetModifiedStateVisibility() const
+	TSharedRef<SWidget> GenerateWidgetForColumn(const FName& ColumnName) override
 	{
-		return (IsModified) ? EVisibility::Visible : EVisibility::Hidden;
+		if (ColumnName == SourceControlFileViewColumn::Name::Id() && bIsModified)
+		{
+			// If the item is modified, wrap the name widget with the modified icon.
+			return SNew(SHorizontalBox)
+				+SHorizontalBox::Slot()
+				.HAlign(HAlign_Left)
+				.AutoWidth()
+				[
+					CommonTableRowType::GenerateWidgetForColumn(ColumnName)
+				]
+				+SHorizontalBox::Slot()
+				.HAlign(HAlign_Right)
+				.VAlign(VAlign_Center)
+				[
+					SNew(SImage)
+					.Image(FAppStyle::GetBrush(TEXT("ContentBrowser.ContentDirty")))
+					.ToolTipText(LOCTEXT("ModifiedFileToolTip", "This file has been modified from the source version"))
+				];
+		}
+
+		return CommonTableRowType::GenerateWidgetForColumn(ColumnName);
 	}
 
-	bool IsSelected;
-	bool IsModified;
-	FString Text;
+private:
+	bool bIsModified = false;
 };
 
 /** Returns whether revert unsaved is enabled */
@@ -129,66 +163,112 @@ public:
 	{
 		ParentFrame = InArgs._ParentWindow.Get();
 
-		for ( TArray<FString>::TConstIterator PackageIter( InArgs._PackagesToRevert.Get() ); PackageIter; ++PackageIter )
-		{
-			ListViewItemSource.Add( MakeShareable(new FRevertCheckBoxListViewItem(*PackageIter) ));
-		}
+		bShowingContentVersePath = FAssetToolsModule::GetModule().Get().ShowingContentVersePath();
 
+		InitializeListViewItemSource(InArgs._PackagesToRevert.Get());
+		SortListViewItemSource();
+
+		TSharedRef<SHeaderRow> HeaderRowWidget = SNew(SHeaderRow);
 
 		ChildSlot
 		[
 			SNew(SBorder)
-			.BorderImage(FAppStyle::GetBrush("ToolPanel.GroupBorder"))
+			.BorderImage(FAppStyle::Get().GetBrush("Brushes.Panel"))
+			.Padding(FMargin(16))
 			[
-
 				SNew(SVerticalBox)
+
 				+SVerticalBox::Slot()
 				.AutoHeight()
-				.Padding(10)
+				.Padding(0.0f, 0.0f, 0.0f, 8.0f)
 				[
 					SNew(STextBlock)
 					.Text(NSLOCTEXT("SourceControl.Revert", "SelectFiles", "Select the files that should be reverted below"))
 				]
+
 				+SVerticalBox::Slot()
-				.AutoHeight()
-				.Padding(10,0)
+				.FillHeight(1.0f)
 				[
-					SNew(SBorder)
-					.BorderImage(FAppStyle::GetBrush("ToolPanel.GroupBorder"))
-					.Padding(5)
-					[
-						SNew(SCheckBox)
-						.OnCheckStateChanged(this, &SSourceControlRevertWidget::ColumnHeaderClicked)
-						.IsEnabled(this, &SSourceControlRevertWidget::OnGetItemsEnabled)
+					SAssignNew(RevertListView, SListViewType)
+					.HeaderRow
+					(
+						SNew(SHeaderRow)
+						+SHeaderRow::Column(SourceControlFileViewColumn::CheckBox::Id())
+						.FixedWidth(38.0f)
 						[
-							SNew(STextBlock)
-							.Text(NSLOCTEXT("SourceControl.Revert", "ListHeader", "File Name"))
+							SNew(SBox)
+							.Padding(FMargin(6.0f, 3.0f, 6.0f, 3.0f))
+							.HAlign(HAlign_Center)
+							[
+								SNew(SCheckBox)
+								.IsChecked(this, &SSourceControlRevertWidget::OnGetColumnHeaderState)
+								.IsEnabled(this, &SSourceControlRevertWidget::OnGetItemsEnabled)
+								.OnCheckStateChanged(this, &SSourceControlRevertWidget::ColumnHeaderClicked)
+							]
 						]
-					]
+
+						+ SHeaderRow::Column(SourceControlFileViewColumn::Icon::Id())
+						.DefaultTooltip(SourceControlFileViewColumn::Icon::GetToolTipText())
+						.FillSized(18.0f)
+						.HeaderContentPadding(FMargin(0.0f))
+						.SortMode(this, &SSourceControlRevertWidget::GetColumnSortMode, SourceControlFileViewColumn::Icon::Id())
+						.OnSort(this, &SSourceControlRevertWidget::OnColumnSortModeChanged)
+						[
+							SNew(SHorizontalBox)
+							+SHorizontalBox::Slot()
+							.Padding(1.0f, 0.0f)
+							[
+								SNew(SBox)
+								.WidthOverride(16.0f)
+								.HeightOverride(16.0f)
+								.HAlign(HAlign_Center)
+								.VAlign(VAlign_Center)
+								.Visibility(this, &SSourceControlRevertWidget::GetIconColumnContentVisibility)
+								[
+									SNew(SImage)
+									.ColorAndOpacity(FSlateColor::UseSubduedForeground())
+									.Image(FRevisionControlStyleManager::Get().GetBrush("RevisionControl.Icon"))
+								]
+							]
+						]
+
+						+SHeaderRow::Column(SourceControlFileViewColumn::Name::Id())
+						.DefaultLabel(LOCTEXT("Asset", "Asset"))
+						.DefaultTooltip(SourceControlFileViewColumn::Name::GetToolTipText())
+						.FillWidth(5.0f)
+						.SortMode(this, &SSourceControlRevertWidget::GetColumnSortMode, SourceControlFileViewColumn::Name::Id())
+						.OnSort(this, &SSourceControlRevertWidget::OnColumnSortModeChanged)
+
+						+SHeaderRow::Column(SourceControlFileViewColumn::Path::Id())
+						.DefaultLabel(LOCTEXT("File", "File"))
+						.DefaultTooltip(SourceControlFileViewColumn::Path::GetToolTipText())
+						.FillWidth(7.0f)
+						.SortMode(this, &SSourceControlRevertWidget::GetColumnSortMode, SourceControlFileViewColumn::Path::Id())
+						.OnSort(this, &SSourceControlRevertWidget::OnColumnSortModeChanged)
+
+						+SHeaderRow::Column(SourceControlFileViewColumn::Type::Id())
+						.DefaultLabel(SourceControlFileViewColumn::Type::GetDisplayText())
+						.DefaultTooltip(SourceControlFileViewColumn::Type::GetToolTipText())
+						.FillWidth(2.0f)
+						.SortMode(this, &SSourceControlRevertWidget::GetColumnSortMode, SourceControlFileViewColumn::Type::Id())
+						.OnSort(this, &SSourceControlRevertWidget::OnColumnSortModeChanged)
+					)
+					.ListItemsSource(&ListViewItemSource)
+					.SelectionMode(ESelectionMode::None)
+					.OnGenerateRow(this, &SSourceControlRevertWidget::OnGenerateRowForList)
 				]
+
 				+SVerticalBox::Slot()
 				.AutoHeight()
-				.Padding(10,0)
-				.MaxHeight(300)
-				[
-					SNew(SBorder)
-					.Padding(5)
-					[
-						SAssignNew(RevertListView, SListViewType)
-						.ListItemsSource(&ListViewItemSource)
-						.OnGenerateRow(this, &SSourceControlRevertWidget::OnGenerateRowForList)
-					]
-				]
-				+SVerticalBox::Slot()
-				.Padding(0, 10, 0, 0)
-				.FillHeight(1)
+				.Padding(0.0f, 16.0f, 0.0f, 0.0f)
 				.VAlign(VAlign_Bottom)
 				.HAlign(HAlign_Fill)
 				[
 					SNew(SHorizontalBox)
+
 					+SHorizontalBox::Slot()
 					.AutoWidth()
-					.Padding(15,5)
+					.Padding(5.0f, 0.0f)
 					.HAlign(HAlign_Left)
 					[
 						SNew(SCheckBox)
@@ -198,29 +278,32 @@ public:
 							.Text(NSLOCTEXT("SourceControl.Revert", "RevertUnchanged", "Revert Unchanged Only"))
 						]
 					]
+
 					+SHorizontalBox::Slot()
 					.HAlign(HAlign_Right)
 					.FillWidth(1)
-					.Padding(5)
 					[
-						SNew(SUniformGridPanel)
-						.SlotPadding(FAppStyle::GetMargin("StandardDialog.SlotPadding"))
-						.MinDesiredSlotWidth(FAppStyle::GetFloat("StandardDialog.MinDesiredSlotWidth"))
-						.MinDesiredSlotHeight(FAppStyle::GetFloat("StandardDialog.MinDesiredSlotHeight"))
-						+SUniformGridPanel::Slot(0,0)
+						SNew(SHorizontalBox)
+
+						+SHorizontalBox::Slot()
+						.AutoWidth()
+						.Padding(5.0f, 0.0f)
 						[
 							SNew(SButton) 
+							.ButtonStyle(&FAppStyle::Get(), "PrimaryButton")
+							.TextStyle(&FAppStyle::Get(), "PrimaryButtonText")
 							.HAlign(HAlign_Center)
-							.ContentPadding(FAppStyle::GetMargin("StandardDialog.ContentPadding"))
 							.OnClicked(this, &SSourceControlRevertWidget::OKClicked)
 							.IsEnabled(this, &SSourceControlRevertWidget::IsOKEnabled)
-							.Text(LOCTEXT("RevertButton", "Revert"))
+							.Text(this, &SSourceControlRevertWidget::GetOkText)
 						]
-						+SUniformGridPanel::Slot(1,0)
+
+						+ SHorizontalBox::Slot()
+						.AutoWidth()
+						.Padding(5.0f, 0.0f)
 						[
 							SNew(SButton) 
 							.HAlign(HAlign_Center)
-							.ContentPadding(FAppStyle::GetMargin("StandardDialog.ContentPadding"))
 							.OnClicked(this, &SSourceControlRevertWidget::CancelClicked)
 							.Text(LOCTEXT("CancelButton", "Cancel"))
 						]
@@ -228,12 +311,6 @@ public:
 				]
 			]
 		];
-
-		// update the modified state of all the files. 
-		UpdateSCCStatus();
-
-		DialogResult = ERevertResults::REVERT_CANCELED;
-		bRevertUnchangedFilesOnly = false;
 	}
 	END_SLATE_FUNCTION_BUILD_OPTIMIZATION
 
@@ -244,61 +321,92 @@ public:
 	 */
 	void GetPackagesToRevert( TArray<FString>& OutPackagesToRevert )
 	{
-		for ( const auto& ListViewItem : ListViewItemSource )
+		if (bRevertUnchangedFilesOnly)
 		{
-			if ((bRevertUnchangedFilesOnly && !ListViewItem->IsModified) || 
-				(!bRevertUnchangedFilesOnly && ListViewItem->IsSelected))
+			for (const TPair<FChangelistTreeItemPtr, TSharedRef<FRevertTableRowState>>& Pair : ListViewItemState)
 			{
-				OutPackagesToRevert.Add(ListViewItem->Text);
+				if (!Pair.Value->bIsModified)
+				{
+					OutPackagesToRevert.Add(Pair.Value->PackageName);
+				}
+			}
+		}
+		else
+		{
+			for (const FChangelistTreeItemPtr& Item : ListViewItemSource)
+			{
+				if (static_cast<IFileViewTreeItem*>(Item.Get())->GetCheckBoxState() == ECheckBoxState::Checked)
+				{
+					const TSharedRef<FRevertTableRowState>* State = ListViewItemState.Find(Item);
+					if (ensure(State))
+					{
+						OutPackagesToRevert.Add((*State)->PackageName);
+					}
+				}
 			}
 		}
 	}
 
 
-	ERevertResults::Type GetResult()
+	ERevertResults GetResult()
 	{
 		return DialogResult;
 	}
 
 private:
+	EColumnSortMode::Type GetColumnSortMode(FName ColumnId) const
+	{
+		if (SortByColumn != ColumnId)
+		{
+			return EColumnSortMode::None;
+		}
+
+		return SortMode;
+	}
 
-	TSharedRef<ITableRow> OnGenerateRowForList( TSharedPtr<FRevertCheckBoxListViewItem> ListItemPtr, const TSharedRef<STableViewBase>& OwnerTable )
+	void OnColumnSortModeChanged(EColumnSortPriority::Type SortPriority, const FName& ColumnId, EColumnSortMode::Type InSortMode)
 	{
-		TSharedPtr<SCheckBox> CheckBox;
-		TSharedRef<ITableRow> Row =
-			SNew(STableRow< TSharedPtr<FString> >, OwnerTable)
-			.IsEnabled(this, &SSourceControlRevertWidget::OnGetItemsEnabled)
-			[
-				SNew(SHorizontalBox)
-				+SHorizontalBox::Slot()
-				.HAlign(HAlign_Left)
-				.AutoWidth()
-				[
-					SAssignNew(CheckBox, SCheckBox)
-					.OnCheckStateChanged(ListItemPtr.ToSharedRef(), &FRevertCheckBoxListViewItem::OnCheckStateChanged)
-					.IsChecked(ListItemPtr.ToSharedRef(), &FRevertCheckBoxListViewItem::OnIsChecked)
-					[
-						SNew(STextBlock)
-						.Text(FText::FromString(ListItemPtr->Text))
-					]
-				]
-				+SHorizontalBox::Slot()
-				.HAlign(HAlign_Right)
-				[
-					SNew(SImage)
-					.Image(FAppStyle::GetBrush(TEXT("ContentBrowser.ContentDirty")))
-					.Visibility(ListItemPtr.ToSharedRef(), &FRevertCheckBoxListViewItem::OnGetModifiedStateVisibility)
-					.ToolTipText(LOCTEXT("ModifiedFileToolTip","This file has been modified from the source version"))
-				]
-			];
+		SortByColumn = ColumnId;
+		SortMode = InSortMode;
+
+		SortListViewItemSource();
+		RevertListView->RequestListRefresh();
+	}
 
-		return Row;
+	EVisibility GetIconColumnContentVisibility() const
+	{
+		// Hide the icon when sorting the icon column (it clashes with the sort mode icon).
+		return GetColumnSortMode(SourceControlFileViewColumn::Icon::Id()) == EColumnSortMode::None ? EVisibility::Visible : EVisibility::Collapsed;
+	}
+
+	TSharedRef<ITableRow> OnGenerateRowForList(FChangelistTreeItemPtr ListItemPtr, const TSharedRef<STableViewBase>& OwnerTable)
+	{
+		TSharedRef<FRevertTableRowState> State = ListViewItemState.FindChecked(ListItemPtr);
+
+		switch (ListItemPtr->GetTreeItemType())
+		{
+		case IChangelistTreeItem::File:
+			return SNew(SSourceControlRevertRow<SFileTableRow>, OwnerTable)
+				.IsEnabled(this, &SSourceControlRevertWidget::OnGetItemsEnabled)
+				.TreeItemToVisualize(ListItemPtr)
+				.IsModified(State->bIsModified)
+				.ShowingContentVersePath(bShowingContentVersePath);
+		case IChangelistTreeItem::OfflineFile:
+			return SNew(SSourceControlRevertRow<SOfflineFileTableRow>, OwnerTable)
+				.IsEnabled(this, &SSourceControlRevertWidget::OnGetItemsEnabled)
+				.TreeItemToVisualize(ListItemPtr)
+				.IsModified(State->bIsModified)
+				.ShowingContentVersePath(bShowingContentVersePath);
+		default:
+			checkNoEntry();
+			return SNew(STableRow<TSharedPtr<FChangelistTreeItemPtr>>, OwnerTable);
+		}
 	}
 
 	/** Called when the settings of the dialog are to be accepted*/
 	FReply OKClicked()
 	{
-		DialogResult = ERevertResults::REVERT_ACCEPTED;
+		DialogResult = ERevertResults::Accepted;
 		ParentFrame.Pin()->RequestDestroyWindow();
 
 		return FReply::Handled();
@@ -308,23 +416,42 @@ private:
 	{
 		if (bRevertUnchangedFilesOnly)
 		{
-			return true;
+			for (const TPair<FChangelistTreeItemPtr, TSharedRef<FRevertTableRowState>>& Pair : ListViewItemState)
+			{
+				if (!Pair.Value->bIsModified)
+				{
+					return true;
+				}
+			}
 		}
-
-		for (int32 i=0; i<ListViewItemSource.Num(); i++)
+		else
 		{
-			if (ListViewItemSource[i]->IsSelected)
+			for (const FChangelistTreeItemPtr& Item : ListViewItemSource)
 			{
-				return true;
+				if (static_cast<IFileViewTreeItem*>(Item.Get())->GetCheckBoxState() == ECheckBoxState::Checked)
+				{
+					return true;
+				}
 			}
 		}
+
 		return false;
 	}
 
+	FText GetOkText() const
+	{
+		if (bRevertUnchangedFilesOnly)
+		{
+			return LOCTEXT("RevertUnchangedButton", "Revert Unchanged");
+		}
+
+		return LOCTEXT("RevertButton", "Revert Selected");
+	}
+
 	/** Called when the settings of the dialog are to be ignored*/
 	FReply CancelClicked()
 	{
-		DialogResult = ERevertResults::REVERT_CANCELED;
+		DialogResult = ERevertResults::Canceled;
 		ParentFrame.Pin()->RequestDestroyWindow();
 
 		return FReply::Handled();
@@ -336,65 +463,117 @@ private:
 		bRevertUnchangedFilesOnly = (NewCheckedState == ECheckBoxState::Checked);
 	}
 
+	ECheckBoxState OnGetColumnHeaderState() const
+	{
+		int32 NumChecked = 0;
+		for (const FChangelistTreeItemPtr& Item : ListViewItemSource)
+		{
+			switch (static_cast<IFileViewTreeItem*>(Item.Get())->GetCheckBoxState())
+			{
+			case ECheckBoxState::Checked:
+				++NumChecked;
+				break;
+			case ECheckBoxState::Undetermined:
+				return ECheckBoxState::Undetermined;
+			}
+		}
+
+		if (NumChecked == 0)
+		{
+			return ECheckBoxState::Unchecked;
+		}
+
+		if (NumChecked == ListViewItemSource.Num())
+		{
+			return ECheckBoxState::Checked;
+		}
+
+		return ECheckBoxState::Undetermined;
+	}
+
 	/**
 	 * Called whenever a column header is clicked, or in the case of the dialog, also when the "Check/Uncheck All" column header
 	 * checkbox is called, because its event bubbles to the column header. 
 	 */
 	void ColumnHeaderClicked( const ECheckBoxState NewCheckedState )
 	{
-		for (int32 i=0; i<ListViewItemSource.Num(); i++)
+		for (const FChangelistTreeItemPtr& Item : ListViewItemSource)
 		{
-			TSharedPtr<FRevertCheckBoxListViewItem> CurListViewItem = ListViewItemSource[i];
-
-			if (OnGetItemsEnabled())
-			{
-				CurListViewItem->IsSelected = (NewCheckedState == ECheckBoxState::Checked);
-			}
+			static_cast<IFileViewTreeItem*>(Item.Get())->SetCheckBoxState(NewCheckedState);
 		}
 	}
 
-	/** Caches the current state of the files, */
-	void UpdateSCCStatus()
+	/** Initializes the current state of the files, */
+	void InitializeListViewItemSource(const TArray<FString>& PackagesToRevert)
 	{
-		TArray<FString> PackagesToCheck;
-		for ( const auto& CurItem : ListViewItemSource )
-		{
-			PackagesToCheck.Add(SourceControlHelpers::PackageFilename(CurItem->Text));
-		}
+		const bool bRevertUnsaved = IsRevertUnsavedEnabled();
+
+		TArray<FString> PackageFilenames = SourceControlHelpers::PackageFilenames(PackagesToRevert);
 
 		// Make sure we update the modified state of the files
 		TSharedRef<FUpdateStatus, ESPMode::ThreadSafe> UpdateStatusOperation = ISourceControlOperation::Create<FUpdateStatus>();
 		UpdateStatusOperation->SetUpdateModifiedState(true);
-		ISourceControlModule::Get().GetProvider().Execute(UpdateStatusOperation, PackagesToCheck);
+		ISourceControlModule::Get().GetProvider().Execute(UpdateStatusOperation, PackageFilenames);
 
 		// Find the files modified from the server version
-		TArray< FSourceControlStateRef > SourceControlStates;
-		ISourceControlModule::Get().GetProvider().GetState( PackagesToCheck, SourceControlStates, EStateCacheUsage::Use );
-
-		ModifiedPackages.Empty();
-
-		const bool bRevertUnsaved = IsRevertUnsavedEnabled();
-		for( const auto& ControlState : SourceControlStates )
+		TArray<FSourceControlStateRef> SourceControlStates;
+		if (ISourceControlModule::Get().GetProvider().GetState(PackageFilenames, SourceControlStates, EStateCacheUsage::Use))
 		{
-			FString PackageName;
-			FPackageName::TryConvertFilenameToLongPackageName(ControlState->GetFilename(), PackageName);
-			for ( const auto& CurItem : ListViewItemSource )
+			ListViewItemSource.Reserve(SourceControlStates.Num());
+			for (const FSourceControlStateRef& SourceControlState : SourceControlStates)
 			{
-				if (CurItem->Text == PackageName)
+				FString PackageName;
+				if (ensure(FPackageName::TryConvertFilenameToLongPackageName(SourceControlState->GetFilename(), PackageName)))
 				{
-					CurItem->IsModified = ControlState->IsModified();
+					const FChangelistTreeItemPtr& ListViewItem = ListViewItemSource.Add_GetRef(MakeShared<FFileTreeItem>(SourceControlState));
 
-					if (bRevertUnsaved)
+					bool bIsModified = SourceControlState->IsModified();
+					if (bRevertUnsaved && !bIsModified)
 					{
-						if (UPackage* Package = FindPackage(NULL, *PackageName))
+						if (UPackage* Package = FindPackage(nullptr, *PackageName))
 						{
 							// If the package contains unsaved changes, it's considered modified as well.
-							CurItem->IsModified |= Package->IsDirty();
+							bIsModified = Package->IsDirty();
 						}
 					}
+
+					ListViewItemState.Emplace(ListViewItem, MakeShared<FRevertTableRowState>(PackageName, bIsModified));
 				}
 			}
 		}
+		else if (bRevertUnsaved)
+		{
+			ListViewItemSource.Reserve(PackagesToRevert.Num());
+			for (int32 Index = 0; Index < PackagesToRevert.Num(); ++Index)
+			{
+				const FChangelistTreeItemPtr& ListViewItem = ListViewItemSource.Add_GetRef(MakeShared<FOfflineFileTreeItem>(PackageFilenames[Index]));
+
+				bool bIsModified = false;
+				if (UPackage* Package = FindPackage(nullptr, *PackagesToRevert[Index]))
+				{
+					// If the package contains unsaved changes, it's considered modified.
+					bIsModified = Package->IsDirty();
+				}
+
+				ListViewItemState.Emplace(ListViewItem, MakeShared<FRevertTableRowState>(PackagesToRevert[Index], bIsModified));
+			}
+		}
+	}
+
+	void SortListViewItemSource()
+	{
+		TFunction<bool(const IFileViewTreeItem&, const IFileViewTreeItem&)> SortPredicate = SourceControlFileViewColumn::GetSortPredicate(
+			SortMode, SortByColumn, bShowingContentVersePath ? SourceControlFileViewColumn::EPathFlags::ShowingVersePath : SourceControlFileViewColumn::EPathFlags::Default);
+		if (SortPredicate)
+		{
+			Algo::SortBy(
+				ListViewItemSource,
+				[](const FChangelistTreeItemPtr& ListViewItem) -> const IFileViewTreeItem&
+				{
+					return static_cast<IFileViewTreeItem&>(*ListViewItem);
+				},
+				SortPredicate);
+		}
 	}
 
 	/** Check for whether the list items are enabled or not */
@@ -404,20 +583,24 @@ private:
 	}
 
 	TWeakPtr<SWindow> ParentFrame;
-	ERevertResults::Type DialogResult;
+	ERevertResults DialogResult = ERevertResults::Canceled;
 
 	/** ListView for the packages the user can revert */
-	typedef SListView<TSharedPtr<FRevertCheckBoxListViewItem>> SListViewType;
+	typedef SListView<FChangelistTreeItemPtr> SListViewType;
 	TSharedPtr<SListViewType> RevertListView;
 
+	FName SortByColumn = SourceControlFileViewColumn::Name::Id();
+	EColumnSortMode::Type SortMode = EColumnSortMode::Ascending;
+
 	/** Collection of items serving as the data source for the list view */
-	TArray<TSharedPtr<FRevertCheckBoxListViewItem>> ListViewItemSource;
+	TArray<FChangelistTreeItemPtr> ListViewItemSource;
+
+	TMap<FChangelistTreeItemPtr, TSharedRef<FRevertTableRowState>> ListViewItemState;
 
-	/** List of package names that are modified from the versions stored in source control; Used as an optimization */
-	TArray<FString> ModifiedPackages;
+	bool bShowingContentVersePath = false;
 
 	/** Flag set by the user to only revert non modified files */
-	bool bRevertUnchangedFilesOnly;
+	bool bRevertUnchangedFilesOnly = false;
 };
 
 bool FSourceControlWindows::PromptForRevert( const TArray<FString>& InPackageNames, bool bInReloadWorld)
@@ -452,7 +635,7 @@ bool FSourceControlWindows::PromptForRevert( const TArray<FString>& InPackageNam
 	{
 		TSharedRef<SWindow> NewWindow = SNew(SWindow)
 			.Title( NSLOCTEXT("SourceControl.RevertWindow", "Title", "Revert Files") )
-			.SizingRule( ESizingRule::Autosized )
+			.ClientSize(FVector2D(640.0f, 492.0f))
 			.SupportsMinimize(false) 
 			.SupportsMaximize(false);
 
@@ -466,7 +649,7 @@ bool FSourceControlWindows::PromptForRevert( const TArray<FString>& InPackageNam
 		FSlateApplication::Get().AddModalWindow(NewWindow, NULL);
 
 		// If the user decided to revert some packages, go ahead and do revert the ones they selected
-		if ( SourceControlWidget->GetResult() == ERevertResults::REVERT_ACCEPTED )
+		if ( SourceControlWidget->GetResult() == ERevertResults::Accepted)
 		{
 			TArray<FString> FinalPackagesToRevert;
 			SourceControlWidget->GetPackagesToRevert(FinalPackagesToRevert);
diff --git a/Engine/Source/Editor/SourceControlWindows/Private/SSourceControlSubmit.cpp b/Engine/Source/Editor/SourceControlWindows/Private/SSourceControlSubmit.cpp
--- a/Engine/Source/Editor/SourceControlWindows/Private/SSourceControlSubmit.cpp
+++ b/Engine/Source/Editor/SourceControlWindows/Private/SSourceControlSubmit.cpp
@@ -2,12 +2,14 @@
 
 #include "SSourceControlSubmit.h"
 
+#include "IAssetTools.h"
 #include "ISourceControlOperation.h"
 #include "SourceControlOperations.h"
 #include "ISourceControlProvider.h"
 #include "ISourceControlModule.h"
 #include "SourceControlHelpers.h"
 #include "SSourceControlCommon.h"
+#include "SSourceControlChangelistRows.h"
 #include "Modules/ModuleManager.h"
 #include "Framework/MultiBox/MultiBoxBuilder.h"
 #include "Widgets/SWindow.h"
@@ -142,40 +144,6 @@ bool TryToVirtualizeFilesToSubmit(const TArray<FSourceControlStateRef>& FileStat
 	}
 }
 
-namespace SSourceControlSubmitWidgetDefs
-{
-	const FName ColumnID_CheckBoxLabel("CheckBox");
-	const FName ColumnID_IconLabel("Icon");
-	const FName ColumnID_AssetLabel("Asset");
-	const FName ColumnID_FileLabel("File");
-
-	const float CheckBoxColumnWidth = 23.0f;
-	const float IconColumnWidth = 21.0f;
-}
-
-
-void SSourceControlSubmitListRow::Construct(const FArguments& InArgs, const TSharedRef<STableViewBase>& InOwnerTableView)
-{
-	SourceControlSubmitWidgetPtr = InArgs._SourceControlSubmitWidget;
-	Item = InArgs._Item;
-
-	SMultiColumnTableRow<TSharedPtr<FFileTreeItem>>::Construct(FSuperRowType::FArguments(), InOwnerTableView);
-}
-
-
-TSharedRef<SWidget> SSourceControlSubmitListRow::GenerateWidgetForColumn(const FName& ColumnName)
-{
-	// Create the widget for this item
-	TSharedPtr<SSourceControlSubmitWidget> SourceControlSubmitWidget = SourceControlSubmitWidgetPtr.Pin();
-	if (SourceControlSubmitWidget.IsValid())
-	{
-		return SourceControlSubmitWidget->GenerateWidgetForItemAndColumn(Item, ColumnName);
-	}
-
-	// Packages dialog no longer valid; return a valid, null widget.
-	return SNullWidget::NullWidget;
-}
-
 
 SSourceControlSubmitWidget::~SSourceControlSubmitWidget()
 {
@@ -186,7 +154,7 @@ SSourceControlSubmitWidget::~SSourceControlSubmitWidget()
 void SSourceControlSubmitWidget::Construct(const FArguments& InArgs)
 {
 	ParentFrame = InArgs._ParentWindow.Get();
-	SortByColumn = SSourceControlSubmitWidgetDefs::ColumnID_AssetLabel;
+	SortByColumn = SourceControlFileViewColumn::Name::Id();
 	SortMode = EColumnSortMode::Ascending;
 	if (!InArgs._Description.Get().IsEmpty())
 	{
@@ -195,6 +163,8 @@ void SSourceControlSubmitWidget::Construct(const FArguments& InArgs)
 	}
 	bAllowSubmit = InArgs._AllowSubmit.Get();
 	bAllowDiffAgainstDepot = InArgs._AllowDiffAgainstDepot.Get();
+	// This widget is only used in a modal window, so bShowingContentVersePath shouldn't change.
+	bShowingContentVersePath = FAssetToolsModule::GetModule().Get().ShowingContentVersePath();
 
 	const bool bDescriptionIsReadOnly = !InArgs._AllowDescriptionChange.Get();
 	const bool bAllowUncheckFiles = InArgs._AllowUncheckFiles.Get();
@@ -204,7 +174,7 @@ void SSourceControlSubmitWidget::Construct(const FArguments& InArgs)
 
 	for (const auto& Item : InArgs._Items.Get())
 	{
-		ListViewItems.Add(MakeShareable(new FFileTreeItem(Item)));
+		ListViewItems.Add(MakeShared<FFileTreeItem>(Item));
 	}
 
 	TSharedRef<SHeaderRow> HeaderRowWidget = SNew(SHeaderRow);
@@ -212,38 +182,62 @@ void SSourceControlSubmitWidget::Construct(const FArguments& InArgs)
 	if (bAllowUncheckFiles)
 	{
 		HeaderRowWidget->AddColumn(
-			SHeaderRow::Column(SSourceControlSubmitWidgetDefs::ColumnID_CheckBoxLabel)
+			SHeaderRow::Column(SourceControlFileViewColumn::CheckBox::Id())
 			[
-				SNew(SCheckBox)
-				.IsChecked(this, &SSourceControlSubmitWidget::GetToggleSelectedState)
-				.OnCheckStateChanged(this, &SSourceControlSubmitWidget::OnToggleSelectedCheckBox)
+				SNew(SBox)
+				.Padding(FMargin(6.0f, 3.0f, 6.0f, 3.0f))
+				.HAlign(HAlign_Center)
+				[
+					SNew(SCheckBox)
+					.IsChecked(this, &SSourceControlSubmitWidget::GetToggleSelectedState)
+					.OnCheckStateChanged(this, &SSourceControlSubmitWidget::OnToggleSelectedCheckBox)
+				]
 			]
-			.FixedWidth(SSourceControlSubmitWidgetDefs::CheckBoxColumnWidth)
+			.FixedWidth(38.0f)
 		);
 	}
 
 	HeaderRowWidget->AddColumn(
-		SHeaderRow::Column(SSourceControlSubmitWidgetDefs::ColumnID_IconLabel)
+		SHeaderRow::Column(SourceControlFileViewColumn::Icon::Id())
 		[
-			SNew(SSpacer)
+			SNew(SHorizontalBox)
+			+SHorizontalBox::Slot()
+			.Padding(1.0f, 0.0f)
+			[
+				SNew(SBox)
+				.WidthOverride(16.0f)
+				.HeightOverride(16.0f)
+				.HAlign(HAlign_Center)
+				.VAlign(VAlign_Center)
+				.Visibility(this, &SSourceControlSubmitWidget::GetIconColumnContentVisibility)
+				[
+					SNew(SImage)
+					.ColorAndOpacity(FSlateColor::UseSubduedForeground())
+					.Image(FRevisionControlStyleManager::Get().GetBrush("RevisionControl.Icon"))
+				]
+			]
 		]
-		.SortMode(this, &SSourceControlSubmitWidget::GetColumnSortMode, SSourceControlSubmitWidgetDefs::ColumnID_IconLabel)
+		.DefaultTooltip(SourceControlFileViewColumn::Icon::GetToolTipText())
+		.SortMode(this, &SSourceControlSubmitWidget::GetColumnSortMode, SourceControlFileViewColumn::Icon::Id())
 		.OnSort(this, &SSourceControlSubmitWidget::OnColumnSortModeChanged)
-		.FixedWidth(SSourceControlSubmitWidgetDefs::IconColumnWidth)
+		.FillSized(18.0f)
+		.HeaderContentPadding(FMargin(0.0f))
 	);
 
 	HeaderRowWidget->AddColumn(
-		SHeaderRow::Column(SSourceControlSubmitWidgetDefs::ColumnID_AssetLabel)
+		SHeaderRow::Column(SourceControlFileViewColumn::Name::Id())
 		.DefaultLabel(LOCTEXT("AssetColumnLabel", "Asset"))
-		.SortMode(this, &SSourceControlSubmitWidget::GetColumnSortMode, SSourceControlSubmitWidgetDefs::ColumnID_AssetLabel)
+		.DefaultTooltip(SourceControlFileViewColumn::Name::GetToolTipText())
+		.SortMode(this, &SSourceControlSubmitWidget::GetColumnSortMode, SourceControlFileViewColumn::Name::Id())
 		.OnSort(this, &SSourceControlSubmitWidget::OnColumnSortModeChanged)
 		.FillWidth(5.0f)
 	);
 
 	HeaderRowWidget->AddColumn(
-		SHeaderRow::Column(SSourceControlSubmitWidgetDefs::ColumnID_FileLabel)
+		SHeaderRow::Column(SourceControlFileViewColumn::Path::Id())
 		.DefaultLabel(LOCTEXT("FileColumnLabel", "File"))
-		.SortMode(this, &SSourceControlSubmitWidget::GetColumnSortMode, SSourceControlSubmitWidgetDefs::ColumnID_FileLabel)
+		.DefaultTooltip(SourceControlFileViewColumn::Path::GetToolTipText())
+		.SortMode(this, &SSourceControlSubmitWidget::GetColumnSortMode, SourceControlFileViewColumn::Path::Id())
 		.OnSort(this, &SSourceControlSubmitWidget::OnColumnSortModeChanged)
 		.FillWidth(7.0f)
 	);
@@ -288,7 +282,7 @@ void SSourceControlSubmitWidget::Construct(const FArguments& InArgs)
 	[
 		SNew(SBorder)
 		[
-			SAssignNew(ListView, SListView<TSharedPtr<FFileTreeItem>>)
+			SAssignNew(ListView, SListView<FChangelistTreeItemPtr>)
 			.ListItemsSource(&ListViewItems)
 			.OnGenerateRow(this, &SSourceControlSubmitWidget::OnGenerateRowForList)
 			.OnContextMenuOpening(this, &SSourceControlSubmitWidget::OnCreateContextMenu)
@@ -538,7 +532,7 @@ bool SSourceControlSubmitWidget::CanDiffAgainstDepot() const
 		const auto& SelectedItems = ListView->GetSelectedItems();
 		if (SelectedItems.Num() == 1)
 		{
-			bCanDiff = SelectedItems[0]->CanDiff();
+			bCanDiff = GetFileItem(SelectedItems[0])->CanDiff();
 		}
 	}
 	return bCanDiff;
@@ -553,12 +547,12 @@ void SSourceControlSubmitWidget::OnDiffAgainstDepot()
 	}
 }
 
-void SSourceControlSubmitWidget::OnDiffAgainstDepotSelected(TSharedPtr<FFileTreeItem> InSelectedItem)
+void SSourceControlSubmitWidget::OnDiffAgainstDepotSelected(FChangelistTreeItemPtr InSelectedItem)
 {
 	if (bAllowDiffAgainstDepot)
 	{
 		FString PackageName;
-		if (FPackageName::TryConvertFilenameToLongPackageName(InSelectedItem->GetFileName().ToString(), PackageName))
+		if (FPackageName::TryConvertFilenameToLongPackageName(GetFileItem(InSelectedItem)->GetFileName().ToString(), PackageName))
 		{
 			TArray<FAssetData> Assets;
 			FAssetRegistryModule& AssetRegistryModule = FModuleManager::LoadModuleChecked<FAssetRegistryModule>(TEXT("AssetRegistry"));
@@ -595,9 +589,9 @@ bool SSourceControlSubmitWidget::CanRevert() const
 	const auto& SelectedItems = ListView->GetSelectedItems();
 	if (SelectedItems.Num() > 0)
 	{
-		return Algo::AllOf(SelectedItems, [](const FFileTreeItemPtr& SelectedItem)
+		return Algo::AllOf(SelectedItems, [](const FChangelistTreeItemPtr& SelectedItem)
 			{
-				return SelectedItem->CanRevert();
+				return GetFileItem(SelectedItem)->CanRevert();
 			}
 		);
 	}
@@ -614,11 +608,11 @@ void SSourceControlSubmitWidget::OnRevert()
 
 	auto RemoveItemsFromListView = [this](TArray<FString>& ItemsToRemove)
 	{
-		ListViewItems.RemoveAll([&ItemsToRemove](const FFileTreeItemPtr& ListViewItem) -> bool
+		ListViewItems.RemoveAll([&ItemsToRemove](const FChangelistTreeItemPtr& ListViewItem) -> bool
 			{
 				return ItemsToRemove.ContainsByPredicate([&ListViewItem](const FString& ItemToRemove) -> bool
 					{
-						return ItemToRemove == ListViewItem->GetFileName().ToString();
+						return ItemToRemove == GetFileItem(ListViewItem)->GetFileName().ToString();
 					}
 				);
 			}
@@ -627,15 +621,16 @@ void SSourceControlSubmitWidget::OnRevert()
 
 	TArray<FString> PackagesToRevert;
 	TArray<FString> FilesToRevert;
-	for (const auto& SelectedItem : SelectedItems)
+	for (const FChangelistTreeItemPtr& SelectedItem : SelectedItems)
 	{
-		if (FPackageName::IsPackageFilename(SelectedItem->GetFileName().ToString()))
+		const FFileTreeItem* FileItem = GetFileItem(SelectedItem);
+		if (FPackageName::IsPackageFilename(FileItem->GetFileName().ToString()))
 		{
-			PackagesToRevert.Add(SelectedItem->GetFileName().ToString());
+			PackagesToRevert.Add(FileItem->GetFileName().ToString());
 		}
 		else
 		{
-			FilesToRevert.Add(SelectedItem->GetFileName().ToString());
+			FilesToRevert.Add(FileItem->GetFileName().ToString());
 		}
 	}
 
@@ -684,89 +679,40 @@ void SSourceControlSubmitWidget::Tick(const FGeometry& AllottedGeometry, const d
 	ISourceControlModule::Get().Tick();
 }
 
-TSharedRef<SWidget> SSourceControlSubmitWidget::GenerateWidgetForItemAndColumn(TSharedPtr<FFileTreeItem> Item, const FName ColumnID) const
+ECheckBoxState SSourceControlSubmitWidget::GetToggleSelectedState() const
 {
-	check(Item.IsValid());
-
-	const FMargin RowPadding(3, 0, 0, 0);
-
-	TSharedPtr<SWidget> ItemContentWidget;
-
-	if (ColumnID == SSourceControlSubmitWidgetDefs::ColumnID_CheckBoxLabel)
+	int32 NumChecked = 0;
+	for (const FChangelistTreeItemPtr& Item : ListViewItems)
 	{
-		ItemContentWidget = SNew(SHorizontalBox)
-			+ SHorizontalBox::Slot()
-			.Padding(RowPadding)
-			[
-				SNew(SCheckBox)
-				.IsChecked(Item.Get(), &FFileTreeItem::GetCheckBoxState)
-				.OnCheckStateChanged(Item.Get(), &FFileTreeItem::SetCheckBoxState)
-			];
-	}
-	else if (ColumnID == SSourceControlSubmitWidgetDefs::ColumnID_IconLabel)
-	{
-		ItemContentWidget = SNew(SHorizontalBox)
-			+ SHorizontalBox::Slot()
-			.HAlign(HAlign_Center)
-			.VAlign(VAlign_Center)
-			[
-				SNew(SImage)
-				.Image(FRevisionControlStyleManager::Get().GetBrush(Item->GetIconName()))
-				.ToolTipText(Item->GetIconTooltip())
-			];
-	}
-	else if (ColumnID == SSourceControlSubmitWidgetDefs::ColumnID_AssetLabel)
-	{
-		ItemContentWidget = SNew(SHorizontalBox)
-			+ SHorizontalBox::Slot()
-			.Padding(RowPadding)
-			[
-				SNew(STextBlock)
-				.Text_Lambda([Item]() { return Item->GetAssetName(); })
-			];
+		switch (GetFileItem(Item)->GetCheckBoxState())
+		{
+		case ECheckBoxState::Checked:
+			++NumChecked;
+			break;
+		case ECheckBoxState::Undetermined:
+			return ECheckBoxState::Undetermined;
+		}
 	}
-	else if (ColumnID == SSourceControlSubmitWidgetDefs::ColumnID_FileLabel)
+
+	if (NumChecked == 0)
 	{
-		ItemContentWidget = SNew(SHorizontalBox)
-			+ SHorizontalBox::Slot()
-			.Padding(RowPadding)
-			[
-				SNew(STextBlock)
-				.Text_Lambda([Item]() { return Item->GetPackageName(); })
-				.ToolTipText(Item->GetFileName())
-			];
+		return ECheckBoxState::Unchecked;
 	}
 
-	return ItemContentWidget.ToSharedRef();
-}
-
-
-ECheckBoxState SSourceControlSubmitWidget::GetToggleSelectedState() const
-{
-	// Default to a Checked state
-	ECheckBoxState PendingState = ECheckBoxState::Checked;
-
-	// Iterate through the list of selected items
-	for (const TSharedPtr<FFileTreeItem>& Item : ListViewItems)
+	if (NumChecked == ListViewItems.Num())
 	{
-		if (Item->GetCheckBoxState() == ECheckBoxState::Unchecked)
-		{
-			// If any item in the list is Unchecked, then represent the entire set of highlighted items as Unchecked,
-			// so that the first (user) toggle of ToggleSelectedCheckBox consistently Checks all items
-			PendingState = ECheckBoxState::Unchecked;
-			break;
-		}
+		return ECheckBoxState::Checked;
 	}
 
-	return PendingState;
+	return ECheckBoxState::Undetermined;
 }
 
 
 void SSourceControlSubmitWidget::OnToggleSelectedCheckBox(ECheckBoxState InNewState)
 {
-	for (const TSharedPtr<FFileTreeItem>& Item : ListViewItems)
+	for (const FChangelistTreeItemPtr& Item : ListViewItems)
 	{
-		Item->SetCheckBoxState(InNewState);
+		GetFileItem(Item)->SetCheckBoxState(InNewState);
 	}
 
 	ListView->RequestListRefresh();
@@ -780,17 +726,18 @@ void SSourceControlSubmitWidget::FillChangeListDescription(FChangeListDescriptio
 	OutDesc.FilesForAdd.Empty();
 	OutDesc.FilesForSubmit.Empty();
 
-	for (const TSharedPtr<FFileTreeItem>& Item : ListViewItems)
+	for (const FChangelistTreeItemPtr& Item : ListViewItems)
 	{
-		if (Item->GetCheckBoxState() == ECheckBoxState::Checked)
+		const FFileTreeItem* FileItem = GetFileItem(Item);
+		if (FileItem->GetCheckBoxState() == ECheckBoxState::Checked)
 		{
-			if (Item->CanCheckIn())
+			if (FileItem->CanCheckIn())
 			{
-				OutDesc.FilesForSubmit.Add(Item->GetFileName().ToString());
+				OutDesc.FilesForSubmit.Add(FileItem->GetFileName().ToString());
 			}
-			else if (Item->NeedsAdding())
+			else if (FileItem->NeedsAdding())
 			{
-				OutDesc.FilesForAdd.Add(Item->GetFileName().ToString());
+				OutDesc.FilesForAdd.Add(FileItem->GetFileName().ToString());
 			}
 		}
 	}
@@ -862,14 +809,12 @@ bool SSourceControlSubmitWidget::CanCheckOut() const
 }
 
 
-TSharedRef<ITableRow> SSourceControlSubmitWidget::OnGenerateRowForList(TSharedPtr<FFileTreeItem> SubmitItem, const TSharedRef<STableViewBase>& OwnerTable)
+TSharedRef<ITableRow> SSourceControlSubmitWidget::OnGenerateRowForList(FChangelistTreeItemPtr SubmitItem, const TSharedRef<STableViewBase>& OwnerTable)
 {
-	TSharedRef<ITableRow> Row =
-	SNew(SSourceControlSubmitListRow, OwnerTable)
-		.SourceControlSubmitWidget(SharedThis(this))
-		.Item(SubmitItem);
-
-	return Row;
+	return SNew(SFileTableRow, OwnerTable)
+		.TreeItemToVisualize(SubmitItem)
+		.PathFlags((bShowingContentVersePath ? SourceControlFileViewColumn::EPathFlags::ShowingVersePath : SourceControlFileViewColumn::EPathFlags::Default) |
+			SourceControlFileViewColumn::EPathFlags::ShowingPackageName);
 }
 
 
@@ -893,6 +838,13 @@ void SSourceControlSubmitWidget::OnColumnSortModeChanged(const EColumnSortPriori
 }
 
 
+EVisibility SSourceControlSubmitWidget::GetIconColumnContentVisibility() const
+{
+	// Hide the icon when sorting the icon column (it clashes with the sort mode icon).
+	return GetColumnSortMode(SourceControlFileViewColumn::Icon::Id()) == EColumnSortMode::None ? EVisibility::Visible : EVisibility::Collapsed;
+}
+
+
 void SSourceControlSubmitWidget::RequestSort()
 {
 	// Sort the list of root items
@@ -904,47 +856,27 @@ void SSourceControlSubmitWidget::RequestSort()
 
 void SSourceControlSubmitWidget::SortTree()
 {
-	if (SortByColumn == SSourceControlSubmitWidgetDefs::ColumnID_AssetLabel)
+	TFunction<bool(const IFileViewTreeItem&, const IFileViewTreeItem&)> SortPredicate = SourceControlFileViewColumn::GetSortPredicate(
+		SortMode, SortByColumn, (bShowingContentVersePath ? SourceControlFileViewColumn::EPathFlags::ShowingVersePath : SourceControlFileViewColumn::EPathFlags::Default) |
+		SourceControlFileViewColumn::EPathFlags::ShowingPackageName);
+	if (SortPredicate)
 	{
-		if (SortMode == EColumnSortMode::Ascending)
-		{
-			ListViewItems.Sort([](const TSharedPtr<FFileTreeItem>& A, const TSharedPtr<FFileTreeItem>& B) {
-				return A->GetAssetName().ToString() < B->GetAssetName().ToString(); });
-		}
-		else if (SortMode == EColumnSortMode::Descending)
-		{
-			ListViewItems.Sort([](const TSharedPtr<FFileTreeItem>& A, const TSharedPtr<FFileTreeItem>& B) {
-				return A->GetAssetName().ToString() >= B->GetAssetName().ToString(); });
-		}
-	}
-	else if (SortByColumn == SSourceControlSubmitWidgetDefs::ColumnID_FileLabel)
-	{
-		if (SortMode == EColumnSortMode::Ascending)
-		{
-			ListViewItems.Sort([](const TSharedPtr<FFileTreeItem>& A, const TSharedPtr<FFileTreeItem>& B) {
-				return A->GetPackageName().ToString() < B->GetPackageName().ToString(); });
-		}
-		else if (SortMode == EColumnSortMode::Descending)
-		{
-			ListViewItems.Sort([](const TSharedPtr<FFileTreeItem>& A, const TSharedPtr<FFileTreeItem>& B) {
-				return A->GetPackageName().ToString() >= B->GetPackageName().ToString(); });
-		}
-	}
-	else if (SortByColumn == SSourceControlSubmitWidgetDefs::ColumnID_IconLabel)
-	{
-		if (SortMode == EColumnSortMode::Ascending)
-		{
-			ListViewItems.Sort([](const TSharedPtr<FFileTreeItem>& A, const TSharedPtr<FFileTreeItem>& B) {
-				return A->GetIconName().ToString() < B->GetIconName().ToString(); });
-		}
-		else if (SortMode == EColumnSortMode::Descending)
-		{
-			ListViewItems.Sort([](const TSharedPtr<FFileTreeItem>& A, const TSharedPtr<FFileTreeItem>& B) {
-				return A->GetIconName().ToString() >= B->GetIconName().ToString(); });
-		}
+		Algo::SortBy(
+			ListViewItems,
+			[](const FChangelistTreeItemPtr& ListViewItem) -> const IFileViewTreeItem&
+			{
+				return *GetFileItem(ListViewItem);
+			},
+			SortPredicate);
 	}
 }
 
+FFileTreeItem* SSourceControlSubmitWidget::GetFileItem(const FChangelistTreeItemPtr& ChangelistItem)
+{
+	check(ChangelistItem->GetTreeItemType() == IChangelistTreeItem::File);
+	return static_cast<FFileTreeItem*>(ChangelistItem.Get());
+}
+
 #undef LOCTEXT_NAMESPACE
 
 #endif // SOURCE_CONTROL_WITH_SLATE
diff --git a/Engine/Source/Editor/SourceControlWindows/Private/SSourceControlSubmit.h b/Engine/Source/Editor/SourceControlWindows/Private/SSourceControlSubmit.h
--- a/Engine/Source/Editor/SourceControlWindows/Private/SSourceControlSubmit.h
+++ b/Engine/Source/Editor/SourceControlWindows/Private/SSourceControlSubmit.h
@@ -93,9 +93,6 @@ public:
 	/** Get dialog result */
 	ESubmitResults::Type GetResult() { return DialogResult; }
 
-	/** Returns a widget representing the item and column supplied */
-	TSharedRef<SWidget> GenerateWidgetForItemAndColumn(TSharedPtr<FFileTreeItem> Item, const FName ColumnID) const;
-
 	/** Gets the requested files and the change list description*/
 	void FillChangeListDescription(FChangeListDescription& OutDesc);
 
@@ -143,7 +140,7 @@ private:
 	bool CanCheckOut() const;
 
 	/** Called by SListView to get a widget corresponding to the supplied item */
-	TSharedRef<ITableRow> OnGenerateRowForList(TSharedPtr<FFileTreeItem> SubmitItemData, const TSharedRef<STableViewBase>& OwnerTable);
+	TSharedRef<ITableRow> OnGenerateRowForList(FChangelistTreeItemPtr SubmitItemData, const TSharedRef<STableViewBase>& OwnerTable);
 
 	/**
 	 * Returns the current column sort mode (ascending or descending) if the ColumnId parameter matches the current
@@ -163,6 +160,8 @@ private:
 	 */
 	void OnColumnSortModeChanged(const EColumnSortPriority::Type SortPriority, const FName& ColumnId, const EColumnSortMode::Type InSortMode);
 
+	EVisibility GetIconColumnContentVisibility() const;
+
 	/**
 	 * Requests that the source list data be sorted according to the current sort column and mode,
 	 * and refreshes the list view.
@@ -178,7 +177,7 @@ private:
 
 	bool CanDiffAgainstDepot() const;
 	void OnDiffAgainstDepot();
-	void OnDiffAgainstDepotSelected(TSharedPtr<FFileTreeItem> InSelectedItem);
+	void OnDiffAgainstDepotSelected(FChangelistTreeItemPtr InSelectedItem);
 
 	/** Called to check whether the selected files in the ListView can be reverted */
 	bool AllowRevert() const;
@@ -187,14 +186,16 @@ private:
 	/** Reverts the files selected in the ListView */
 	void OnRevert();
 
+	static FFileTreeItem* GetFileItem(const FChangelistTreeItemPtr& ChangelistItem);
+
 private:
 	ESubmitResults::Type DialogResult;
 
 	/** ListBox for selecting which object to consolidate */
-	TSharedPtr<SListView<TSharedPtr<FFileTreeItem>>> ListView;
+	TSharedPtr<SListView<FChangelistTreeItemPtr>> ListView;
 
 	/** Collection of objects (Widgets) to display in the List View. */
-	TArray<TSharedPtr<FFileTreeItem>> ListViewItems;
+	TArray<FChangelistTreeItemPtr> ListViewItems;
 
 	/** Pointer to the parent modal window */
 	TWeakPtr<SWindow> ParentFrame;
@@ -211,38 +212,11 @@ private:
 	/** Whether a diff against the depot may be performed from within the submit dialog */
 	bool bAllowDiffAgainstDepot;
 
+	bool bShowingContentVersePath;
+
 	/** Specify which column to sort with */
 	FName SortByColumn;
 
 	/** Currently selected sorting mode */
 	EColumnSortMode::Type SortMode;
 };
-
-class SSourceControlSubmitListRow : public SMultiColumnTableRow<TSharedPtr<FFileTreeItem>>
-{
-public:
-
-	SLATE_BEGIN_ARGS(SSourceControlSubmitListRow) {}
-
-	/** The SSourceControlSubmitWidget that owns the tree.  We'll only keep a weak reference to it. */
-	SLATE_ARGUMENT(TSharedPtr<SSourceControlSubmitWidget>, SourceControlSubmitWidget)
-
-		/** The list item for this row */
-		SLATE_ARGUMENT(TSharedPtr<FFileTreeItem>, Item)
-
-	SLATE_END_ARGS()
-
-	/** Construct function for this widget */
-	void Construct(const FArguments& InArgs, const TSharedRef<STableViewBase>& InOwnerTableView);
-
-	/** Overridden from SMultiColumnTableRow.  Generates a widget for this column of the list row. */
-	virtual TSharedRef<SWidget> GenerateWidgetForColumn(const FName& ColumnName) override;
-
-private:
-
-	/** Weak reference to the SSourceControlSubmitWidget that owns our list */
-	TWeakPtr<SSourceControlSubmitWidget> SourceControlSubmitWidgetPtr;
-
-	/** The item associated with this row of data */
-	TSharedPtr<FFileTreeItem> Item;
-};
diff --git a/Engine/Source/Editor/SourceControlWindows/SourceControlWindows.Build.cs b/Engine/Source/Editor/SourceControlWindows/SourceControlWindows.Build.cs
--- a/Engine/Source/Editor/SourceControlWindows/SourceControlWindows.Build.cs
+++ b/Engine/Source/Editor/SourceControlWindows/SourceControlWindows.Build.cs
@@ -11,7 +11,8 @@ public class SourceControlWindows : ModuleRules
 				"Core", 
 				"CoreUObject", 
                 "InputCore",
-				"Engine", 
+				"Engine",
+				"Projects",
 				"Slate",
 				"SlateCore",
 				"SourceControl", 
