Updated the source control submit dialog to display the same value for the Path column as other source control dialogs.
Updated the checked out dialog to use the common source control row and support sorting on all columns.

diff --git a/Engine/Source/Editor/SourceControlWindows/Private/SSourceControlChangelistRows.cpp b/Engine/Source/Editor/SourceControlWindows/Private/SSourceControlChangelistRows.cpp
--- a/Engine/Source/Editor/SourceControlWindows/Private/SSourceControlChangelistRows.cpp
+++ b/Engine/Source/Editor/SourceControlWindows/Private/SSourceControlChangelistRows.cpp
@@ -147,14 +147,7 @@ namespace SourceControlFileViewColumn
 		}
 		else if (ColumnId == SourceControlFileViewColumn::Path::Id())
 		{
-			if (PathFlags == EPathFlags::ShowingPackageName)
-			{
-				return [](const IFileViewTreeItem& Lhs, const IFileViewTreeItem& Rhs)
-				{
-					return Lhs.GetPackageName().Compare(Rhs.GetPackageName());
-				};
-			}
-			else if (PathFlags == EPathFlags::ShowingVersePath)
+			if (PathFlags == EPathFlags::ShowingVersePath)
 			{
 				return [](const IFileViewTreeItem& Lhs, const IFileViewTreeItem& Rhs)
 				{
@@ -173,25 +166,6 @@ namespace SourceControlFileViewColumn
 					return Lhs.GetPath().Compare(Rhs.GetPath());
 				};
 			}
-			else if (PathFlags == (EPathFlags::ShowingPackageName | EPathFlags::ShowingVersePath))
-			{
-				return [](const IFileViewTreeItem& Lhs, const IFileViewTreeItem& Rhs)
-				{
-					if (Lhs.GetVersePath().IsValid() && Rhs.GetVersePath().IsValid())
-					{
-						return Lhs.GetVersePath().Compare(Rhs.GetVersePath());
-					}
-					else if (Lhs.GetVersePath().IsValid())
-					{
-						return Lhs.GetVersePath().ToString().Compare(Rhs.GetPackageName());
-					}
-					else if (Rhs.GetVersePath().IsValid())
-					{
-						return Lhs.GetPackageName().Compare(Rhs.GetVersePath().ToString());
-					}
-					return Lhs.GetPackageName().Compare(Rhs.GetPackageName());
-				};
-			}
 
 			return [](const IFileViewTreeItem& Lhs, const IFileViewTreeItem& Rhs)
 			{
@@ -674,10 +648,6 @@ void SFileTableRow::PopulateSearchString(const FFileTreeItem& Item, SourceContro
 	{
 		OutStrings.Emplace(Item.GetAssetVersePath().ToString());
 	}
-	else if (EnumHasAnyFlags(PathFlags, SourceControlFileViewColumn::EPathFlags::ShowingPackageName))
-	{
-		OutStrings.Emplace(Item.GetAssetPackageName().ToString()); // Path
-	}
 	else
 	{
 		OutStrings.Emplace(Item.GetAssetPath().ToString()); // Path
@@ -713,10 +683,6 @@ FText SFileTableRow::GetDisplayPath() const
 	{
 		return TreeItem->GetAssetVersePath();
 	}
-	else if (EnumHasAnyFlags(PathFlags.Get(SourceControlFileViewColumn::EPathFlags::Default), SourceControlFileViewColumn::EPathFlags::ShowingPackageName))
-	{
-		return TreeItem->GetAssetPackageName();
-	}
 	return TreeItem->GetAssetPath();
 }
 
@@ -915,7 +881,7 @@ TSharedRef<SWidget> SOfflineFileTableRow::GenerateWidgetForColumn(const FName& C
 	}
 	else if (ColumnId == SourceControlFileViewColumn::Discard::Id())
 	{
-		FString Filename = TreeItem->GetFullPathname();
+		const FString& Filename = TreeItem->GetFullPathname();
 		if (!FUnsavedAssetsTrackerModule::Get().IsAssetUnsaved(Filename))
 		{
 			return SNew(SBox)
@@ -932,24 +898,19 @@ TSharedRef<SWidget> SOfflineFileTableRow::GenerateWidgetForColumn(const FName& C
 				{
 					// Normalize packagenames and filenames
 					FString PackageName;
+					if (FPackageName::TryConvertFilenameToLongPackageName(TreeItem->GetFullPathname(), PackageName))
 					{
-						FString TreeName = TreeItem->GetPackageName();
-
-						if (!FPackageName::TryConvertFilenameToLongPackageName(TreeName, PackageName))
+						// Validate we have a saved map
+						if (UPackage* Package = FindPackage(nullptr, *PackageName))
 						{
-							PackageName = MoveTemp(TreeName);
-						}
-					}
-					// Validate we have a saved map
-					if (UPackage* Package = FindPackage(nullptr, *PackageName))
-					{
-						UPackage* LevelPackage = Package->GetOutermost();
-						if (LevelPackage == GetTransientPackage()
-							|| LevelPackage->HasAnyFlags(RF_Transient)
-							|| !FPackageName::IsValidLongPackageName(LevelPackage->GetName()))
-						{
-							FMessageDialog::Open(EAppMsgType::Ok, LOCTEXT("DiscardUnsavedChangesSaveMap", "You need to save the level before discarding unsaved changes."));
-							return FReply::Handled();
+							UPackage* LevelPackage = Package->GetOutermost();
+							if (LevelPackage == GetTransientPackage()
+								|| LevelPackage->HasAnyFlags(RF_Transient)
+								|| !FPackageName::IsValidLongPackageName(LevelPackage->GetName()))
+							{
+								FMessageDialog::Open(EAppMsgType::Ok, LOCTEXT("DiscardUnsavedChangesSaveMap", "You need to save the level before discarding unsaved changes."));
+								return FReply::Handled();
+							}
 						}
 					}
 					
@@ -957,11 +918,11 @@ TSharedRef<SWidget> SOfflineFileTableRow::GenerateWidgetForColumn(const FName& C
 					ExecuteOnGameThread(UE_SOURCE_LOCATION,
 						[this]
 						{
-							TArray<FString> PackageToReload { TreeItem->GetPackageName() };
+							TArray<FString> Filenames { TreeItem->GetFullPathname() };
 							const bool bAllowReloadWorld = true;
 							const bool bInteractive = false;
 							USourceControlHelpers::ApplyOperationAndReloadPackages(
-								PackageToReload,
+								Filenames,
 								[](const TArray<FString>&) -> bool { return true; },
 								bAllowReloadWorld,
 								bInteractive
@@ -1001,10 +962,6 @@ void SOfflineFileTableRow::PopulateSearchString(const FOfflineFileTreeItem& Item
 	{
 		OutStrings.Emplace(Item.GetDisplayVersePath().ToString());
 	}
-	else if (EnumHasAnyFlags(PathFlags, SourceControlFileViewColumn::EPathFlags::ShowingPackageName))
-	{
-		OutStrings.Emplace(Item.GetDisplayPackageName().ToString()); // Path
-	}
 	else
 	{
 		OutStrings.Emplace(Item.GetDisplayPath().ToString()); // Path
@@ -1030,7 +987,7 @@ FText SOfflineFileTableRow::GetDisplayName() const
 
 FText SOfflineFileTableRow::GetFilename() const
 {
-	return TreeItem->GetDisplayPackageName();
+	return TreeItem->GetDisplayFilename();
 }
 
 FText SOfflineFileTableRow::GetDisplayPath() const
@@ -1039,10 +996,6 @@ FText SOfflineFileTableRow::GetDisplayPath() const
 	{
 		return TreeItem->GetDisplayVersePath();
 	}
-	else if (EnumHasAnyFlags(PathFlags.Get(SourceControlFileViewColumn::EPathFlags::Default), SourceControlFileViewColumn::EPathFlags::ShowingPackageName))
-	{
-		return TreeItem->GetDisplayPackageName();
-	}
 	return TreeItem->GetDisplayPath();
 }
 
diff --git a/Engine/Source/Editor/SourceControlWindows/Private/SSourceControlChangelistRows.h b/Engine/Source/Editor/SourceControlWindows/Private/SSourceControlChangelistRows.h
--- a/Engine/Source/Editor/SourceControlWindows/Private/SSourceControlChangelistRows.h
+++ b/Engine/Source/Editor/SourceControlWindows/Private/SSourceControlChangelistRows.h
@@ -296,6 +296,7 @@ public:
 	SLATE_BEGIN_ARGS(SOfflineFileTableRow)
 		: _TreeItemToVisualize()
 		, _HighlightText()
+		, _PathFlags(SourceControlFileViewColumn::EPathFlags::Default)
 	{
 	}
 		SLATE_ARGUMENT(FChangelistTreeItemPtr, TreeItemToVisualize)
diff --git a/Engine/Source/Editor/SourceControlWindows/Private/SSourceControlChangelists.cpp b/Engine/Source/Editor/SourceControlWindows/Private/SSourceControlChangelists.cpp
--- a/Engine/Source/Editor/SourceControlWindows/Private/SSourceControlChangelists.cpp
+++ b/Engine/Source/Editor/SourceControlWindows/Private/SSourceControlChangelists.cpp
@@ -1613,7 +1613,7 @@ void SSourceControlChangelistsWidget::SetSelectedFiles(const TArray<FString>& Fi
 			{
 				if (ChildItem->GetTreeItemType() == IChangelistTreeItem::File)
 				{
-					const FString& ChildFilename = StaticCastSharedPtr<FFileTreeItem>(ChildItem)->FileState->GetFilename();
+					const FString& ChildFilename = StaticCastSharedPtr<FFileTreeItem>(ChildItem)->GetFullPathname();
 					if (ChildFilename.Compare(Filename, ESearchCase::IgnoreCase) == 0)
 					{
 						return Item;
@@ -1626,17 +1626,9 @@ void SSourceControlChangelistsWidget::SetSelectedFiles(const TArray<FString>& Fi
 		{
 			for (const TSharedPtr<IChangelistTreeItem>& ChildItem : Item->GetChildren())
 			{
-				if (ChildItem->GetTreeItemType() == IChangelistTreeItem::File)
+				if (ChildItem->GetTreeItemType() == IChangelistTreeItem::File || ChildItem->GetTreeItemType() == IChangelistTreeItem::OfflineFile)
 				{
-					const FString& ChildFilename = StaticCastSharedPtr<FFileTreeItem>(ChildItem)->FileState->GetFilename();
-					if (ChildFilename.Compare(Filename, ESearchCase::IgnoreCase) == 0)
-					{
-						return Item;
-					}
-				}
-				else if (ChildItem->GetTreeItemType() == IChangelistTreeItem::OfflineFile)
-				{
-					const FString& ChildFilename = StaticCastSharedPtr<FOfflineFileTreeItem>(ChildItem)->GetFilename();
+					const FString& ChildFilename = StaticCastSharedPtr<IFileViewTreeItem>(ChildItem)->GetFullPathname();
 					if (ChildFilename.Compare(Filename, ESearchCase::IgnoreCase) == 0)
 					{
 						return Item;
@@ -1672,18 +1664,9 @@ void SSourceControlChangelistsWidget::SetSelectedFiles(const TArray<FString>& Fi
 
 		for (const TSharedPtr<IChangelistTreeItem>& ChildItem : FoundChangelistTreeItem->GetChildren())
 		{
-			if (ChildItem->GetTreeItemType() == IChangelistTreeItem::File)
+			if (ChildItem->GetTreeItemType() == IChangelistTreeItem::File || ChildItem->GetTreeItemType() == IChangelistTreeItem::OfflineFile)
 			{
-				const FString& ChildFilename = StaticCastSharedPtr<FFileTreeItem>(ChildItem)->FileState->GetFilename().ToLower();
-				if (FilenamesLowerCase.Contains(ChildFilename))
-				{
-					GetActiveFileListView().SetItemSelection(ChildItem, /*bSelected*/true);
-					LastFoundItem = ChildItem;
-				}
-			}
-			else if (ChildItem->GetTreeItemType() == IChangelistTreeItem::OfflineFile)
-			{
-				const FString& ChildFilename = StaticCastSharedPtr<FOfflineFileTreeItem>(ChildItem)->GetFilename().ToLower();
+				const FString ChildFilename = StaticCastSharedPtr<IFileViewTreeItem>(ChildItem)->GetFullPathname().ToLower();
 				if (FilenamesLowerCase.Contains(ChildFilename))
 				{
 					GetActiveFileListView().SetItemSelection(ChildItem, /*bSelected*/true);
@@ -1791,12 +1774,12 @@ void SSourceControlChangelistsWidget::GetSelectedFiles(TArray<FString>& OutContr
 			{
 				if (Parent->GetTreeItemType() == IChangelistTreeItem::UncontrolledChangelist)
 				{
-					const FString& Filename = StaticCastSharedPtr<FOfflineFileTreeItem>(Item)->GetFilename();
+					const FString& Filename = StaticCastSharedPtr<FOfflineFileTreeItem>(Item)->GetFullPathname();
 					OutUncontrolledFiles.Add(Filename);
 				}
 				else if (Parent->GetTreeItemType() == IChangelistTreeItem::UnsavedAssets)
 				{
-					const FString& Filename = StaticCastSharedPtr<FOfflineFileTreeItem>(Item)->GetFilename();
+					const FString& Filename = StaticCastSharedPtr<FOfflineFileTreeItem>(Item)->GetFullPathname();
 					OutControlledFiles.Add(Filename);
 				}
 			}
@@ -3595,7 +3578,7 @@ FReply SSourceControlChangelistsWidget::OnUnsavedAssetsDragged(const FGeometry&
 			if (InTreeItem->GetTreeItemType() == IChangelistTreeItem::OfflineFile)
 			{
 				TSharedRef<FOfflineFileTreeItem> FileTreeItem = StaticCastSharedRef<FOfflineFileTreeItem>(InTreeItem.ToSharedRef());
-				Operation->OfflineFiles.Emplace(FileTreeItem->GetFilename());
+				Operation->OfflineFiles.Emplace(FileTreeItem->GetFullPathname());
 			}
 		}
 
@@ -3630,14 +3613,9 @@ void SSourceControlChangelistsWidget::OnGetChangelistChildren(FChangelistTreeIte
 
 void SSourceControlChangelistsWidget::OnItemDoubleClicked(TSharedPtr<IChangelistTreeItem> Item)
 {
-	if (Item->GetTreeItemType() == IChangelistTreeItem::OfflineFile)
-	{
-		const FString& Filename = StaticCastSharedPtr<FOfflineFileTreeItem>(Item)->GetFilename();
-		ISourceControlWindowsModule::Get().OnChangelistFileDoubleClicked().Broadcast(Filename);
-	}
-	else if (Item->GetTreeItemType() == IChangelistTreeItem::File)
+	if (Item->GetTreeItemType() == IChangelistTreeItem::File || Item->GetTreeItemType() == IChangelistTreeItem::OfflineFile)
 	{
-		const FString& Filename = StaticCastSharedPtr<FFileTreeItem>(Item)->FileState->GetFilename();
+		const FString& Filename = StaticCastSharedPtr<IFileViewTreeItem>(Item)->GetFullPathname();
 		ISourceControlWindowsModule::Get().OnChangelistFileDoubleClicked().Broadcast(Filename);
 	}
 	else if (Item->GetTreeItemType() == IChangelistTreeItem::Changelist)
diff --git a/Engine/Source/Editor/SourceControlWindows/Private/SSourceControlCheckedOutDialog.cpp b/Engine/Source/Editor/SourceControlWindows/Private/SSourceControlCheckedOutDialog.cpp
--- a/Engine/Source/Editor/SourceControlWindows/Private/SSourceControlCheckedOutDialog.cpp
+++ b/Engine/Source/Editor/SourceControlWindows/Private/SSourceControlCheckedOutDialog.cpp
@@ -1,6 +1,7 @@
 // Copyright Epic Games, Inc. All Rights Reserved.
 
 #include "SSourceControlCheckedOutDialog.h"
+#include "SSourceControlChangelistRows.h"
 #include "AssetToolsModule.h"
 #include "Brushes/SlateDynamicImageBrush.h"
 #include "Styling/AppStyle.h"
@@ -25,17 +26,10 @@
 
 #define LOCTEXT_NAMESPACE "SSourceControlConflict"
 
-namespace SSourceControlConflictWarningWidgetDefs
-{
-	const FName ColumnID_AssetLabel("Asset");
-	const FName ColumnID_AssetClassLabel("AssetClass");
-	const FName ColumnID_UserNameLabel("UserName");
-}
-
 void SSourceControlCheckedOutDialog::Construct(const FArguments& InArgs)
 {
 	ParentFrame = InArgs._ParentWindow;
-	SortByColumn = SSourceControlConflictWarningWidgetDefs::ColumnID_AssetLabel;
+	SortByColumn = SourceControlFileViewColumn::Name::Id();
 	SortMode = EColumnSortMode::Ascending;
 	bShowingContentVersePath = FAssetToolsModule::GetModule().Get().ShowingContentVersePath();
 
@@ -55,9 +49,9 @@ void SSourceControlCheckedOutDialog::Construct(const FArguments& InArgs)
 	if (bShowColumnUserName)
 	{
 		bool bAnyCheckedOut = false;
-		for (auto& Item : ListViewItems)
+		for (const FChangelistTreeItemPtr& Item : ListViewItems)
 		{
-			if (!Item->GetCheckedOutByUser().IsEmpty())
+			if (!static_cast<IFileViewTreeItem&>(*Item).GetCheckedOutBy().IsEmpty())
 			{
 				bAnyCheckedOut = true;
 				break;
@@ -69,29 +63,45 @@ void SSourceControlCheckedOutDialog::Construct(const FArguments& InArgs)
 	if (bShowColumnAssetName)
 	{
 		HeaderRowWidget->AddColumn(
-			SHeaderRow::Column(SSourceControlConflictWarningWidgetDefs::ColumnID_AssetLabel)
-			.DefaultLabel(LOCTEXT("AssetColumnLabel", "Asset Name"))
-			.SortMode(this, &SSourceControlCheckedOutDialog::GetColumnSortMode, SSourceControlConflictWarningWidgetDefs::ColumnID_AssetLabel)
+			SHeaderRow::Column(SourceControlFileViewColumn::Name::Id())
+			.DefaultLabel(LOCTEXT("AssetColumnLabel", "Asset"))
+			.DefaultTooltip(SourceControlFileViewColumn::Name::GetToolTipText())
+			.SortMode(this, &SSourceControlCheckedOutDialog::GetColumnSortMode, SourceControlFileViewColumn::Name::Id())
 			.OnSort(this, &SSourceControlCheckedOutDialog::OnColumnSortModeChanged)
-			.FillWidth(0.5f)
+			.FillWidth(1.5f)
+		);
+
+		HeaderRowWidget->AddColumn(
+			SHeaderRow::Column(SourceControlFileViewColumn::Path::Id())
+			.DefaultLabel(LOCTEXT("FileColumnLabel", "File"))
+			.DefaultTooltip(SourceControlFileViewColumn::Path::GetToolTipText())
+			.SortMode(this, &SSourceControlCheckedOutDialog::GetColumnSortMode, SourceControlFileViewColumn::Path::Id())
+			.OnSort(this, &SSourceControlCheckedOutDialog::OnColumnSortModeChanged)
+			.FillWidth(3.0f)
 		);
 	}
 
 	if (bShowColumnAssetClass)
 	{
 		HeaderRowWidget->AddColumn(
-			SHeaderRow::Column(SSourceControlConflictWarningWidgetDefs::ColumnID_AssetClassLabel)
-			.DefaultLabel(LOCTEXT("AssetClassLabel", "Asset Class"))
-			.FillWidth(0.5f)
+			SHeaderRow::Column(SourceControlFileViewColumn::Type::Id())
+			.DefaultLabel(SourceControlFileViewColumn::Type::GetDisplayText())
+			.DefaultTooltip(SourceControlFileViewColumn::Type::GetToolTipText())
+			.SortMode(this, &SSourceControlCheckedOutDialog::GetColumnSortMode, SourceControlFileViewColumn::Type::Id())
+			.OnSort(this, &SSourceControlCheckedOutDialog::OnColumnSortModeChanged)
+			.FillWidth(1.0f)
 		);
 	}
 
 	if (bShowColumnUserName)
 	{
 		HeaderRowWidget->AddColumn(
-			SHeaderRow::Column(SSourceControlConflictWarningWidgetDefs::ColumnID_UserNameLabel)
-			.DefaultLabel(LOCTEXT("UserNameLabel", "User Name"))
-			.FillWidth(0.25f)
+			SHeaderRow::Column(SourceControlFileViewColumn::CheckedOutByUser::Id())
+			.DefaultLabel(SourceControlFileViewColumn::CheckedOutByUser::GetDisplayText())
+			.DefaultTooltip(SourceControlFileViewColumn::CheckedOutByUser::GetToolTipText())
+			.SortMode(this, &SSourceControlCheckedOutDialog::GetColumnSortMode, SourceControlFileViewColumn::CheckedOutByUser::Id())
+			.OnSort(this, &SSourceControlCheckedOutDialog::OnColumnSortModeChanged)
+			.FillWidth(1.0f)
 		);
 	}
 
@@ -194,7 +204,7 @@ void SSourceControlCheckedOutDialog::Construct(const FArguments& InArgs)
 			.HeightOverride(200.0f)
 			.WidthOverride(800.0f)
 			[
-				SAssignNew(ListView, SListView<FFileTreeItemRef>)
+				SAssignNew(ListView, SListView<FChangelistTreeItemPtr>)
 				.ListItemsSource(&ListViewItems)
 				.OnGenerateRow(this, &SSourceControlCheckedOutDialog::OnGenerateRowForList)
 				.HeaderRow(HeaderRowWidget)
@@ -215,12 +225,12 @@ void SSourceControlCheckedOutDialog::Construct(const FArguments& InArgs)
 	RequestSort();
 }
 
-TSharedRef<ITableRow> SSourceControlCheckedOutDialog::OnGenerateRowForList(FFileTreeItemRef InItem, const TSharedRef<STableViewBase>& OwnerTable)
+TSharedRef<ITableRow> SSourceControlCheckedOutDialog::OnGenerateRowForList(FChangelistTreeItemPtr InItem, const TSharedRef<STableViewBase>& OwnerTable)
 {
 	TSharedRef<ITableRow> Row =
-		SNew(SSourceControlCheckedOutDialogListRow, OwnerTable)
-		.Item(InItem)
-		.ShowingContentVersePath(bShowingContentVersePath);
+		SNew(SFileTableRow, OwnerTable)
+		.TreeItemToVisualize(InItem)
+		.PathFlags(bShowingContentVersePath ? SourceControlFileViewColumn::EPathFlags::ShowingVersePath : SourceControlFileViewColumn::EPathFlags::Default);
 
 	return Row;
 }
@@ -252,18 +262,17 @@ void SSourceControlCheckedOutDialog::RequestSort()
 
 void SSourceControlCheckedOutDialog::SortTree()
 {
-	if (SortByColumn == SSourceControlConflictWarningWidgetDefs::ColumnID_AssetLabel)
+	TFunction<bool(const IFileViewTreeItem&, const IFileViewTreeItem&)> SortPredicate = SourceControlFileViewColumn::GetSortPredicate(
+		SortMode, SortByColumn, bShowingContentVersePath ? SourceControlFileViewColumn::EPathFlags::ShowingVersePath : SourceControlFileViewColumn::EPathFlags::Default);
+	if (SortPredicate)
 	{
-		if (SortMode == EColumnSortMode::Ascending)
-		{
-			ListViewItems.Sort([this](const FFileTreeItemRef& A, const FFileTreeItemRef& B) {
-				return A->GetAssetName().CompareTo(B->GetAssetName()) < 0; });
-		}
-		else if (SortMode == EColumnSortMode::Descending)
-		{
-			ListViewItems.Sort([this](const FFileTreeItemRef& A, const FFileTreeItemRef& B) {
-				return A->GetAssetName().CompareTo(B->GetAssetName()) >= 0; });
-		}
+		Algo::SortBy(
+			ListViewItems,
+			[](const FChangelistTreeItemPtr& ListViewItem) -> const IFileViewTreeItem&
+			{
+				return static_cast<IFileViewTreeItem&>(*ListViewItem);
+			},
+			SortPredicate);
 	}
 }
 
@@ -284,57 +293,4 @@ FReply SSourceControlCheckedOutDialog::OnKeyDown(const FGeometry& MyGeometry, co
 	return FReply::Unhandled();
 }
 
-void SSourceControlCheckedOutDialogListRow::Construct(const FArguments& InArgs, const TSharedRef<STableViewBase>& InOwnerTableView)
-{
-	Item = InArgs._Item;
-	bShowingContentVersePath = InArgs._ShowingContentVersePath;
-
-	SMultiColumnTableRow<FFileTreeItemRef>::Construct(FSuperRowType::FArguments(), InOwnerTableView);
-}
-
-TSharedRef<SWidget> SSourceControlCheckedOutDialogListRow::GenerateWidgetForColumn(const FName& ColumnName)
-{
-	const FMargin RowPadding(8, 2, 2, 2);
-
-	TSharedPtr<SWidget> ItemContentWidget;
-	if (ColumnName == SSourceControlConflictWarningWidgetDefs::ColumnID_AssetLabel)
-	{
-		ItemContentWidget = SNew(SHorizontalBox)
-			+ SHorizontalBox::Slot()
-			.Padding(RowPadding)
-			.VAlign(VAlign_Center)
-			[
-				SNew(STextBlock)
-				.Text(Item->GetAssetName())
-				.ToolTipText(bShowingContentVersePath && Item->GetVersePath().IsValid() ? Item->GetAssetVersePath() : Item->GetAssetPackageName())
-			];
-	}
-	else if (ColumnName == SSourceControlConflictWarningWidgetDefs::ColumnID_AssetClassLabel)
-	{
-		ItemContentWidget = SNew(SHorizontalBox)
-			+ SHorizontalBox::Slot()
-			.Padding(RowPadding)
-			.VAlign(VAlign_Center)
-			[
-				SNew(STextBlock)
-				.Text(Item->GetAssetTypeName())
-				.ToolTipText(Item->GetAssetType())
-			];
-	}
-	else if (ColumnName == SSourceControlConflictWarningWidgetDefs::ColumnID_UserNameLabel)
-	{
-		ItemContentWidget = SNew(SHorizontalBox)
-			+ SHorizontalBox::Slot()
-			.Padding(RowPadding)
-			.VAlign(VAlign_Center)
-			[
-				SNew(STextBlock)
-				.Text(Item->GetCheckedOutByUser())
-				.ToolTipText(Item->GetCheckedOutByUser())
-			];
-	}
-
-	return ItemContentWidget.ToSharedRef();
-}
-
 #undef LOCTEXT_NAMESPACE
diff --git a/Engine/Source/Editor/SourceControlWindows/Private/SSourceControlCheckedOutDialog.h b/Engine/Source/Editor/SourceControlWindows/Private/SSourceControlCheckedOutDialog.h
--- a/Engine/Source/Editor/SourceControlWindows/Private/SSourceControlCheckedOutDialog.h
+++ b/Engine/Source/Editor/SourceControlWindows/Private/SSourceControlCheckedOutDialog.h
@@ -48,7 +48,7 @@ public:
 
 private:
 	/** Callback to generate ListBoxRows */
-	TSharedRef<ITableRow> OnGenerateRowForList(FFileTreeItemRef InItem, const TSharedRef<STableViewBase>& OwnerTable);
+	TSharedRef<ITableRow> OnGenerateRowForList(FChangelistTreeItemPtr InItem, const TSharedRef<STableViewBase>& OwnerTable);
 
 	/**
 	 * Returns the current column sort mode (ascending or descending) if the ColumnId parameter matches the current
@@ -90,10 +90,10 @@ private:
 	TWeakPtr<SWindow> ParentFrame;
 
 	/** Collection of objects to display in the List View */
-	TArray<FFileTreeItemRef> ListViewItems;
+	TArray<FChangelistTreeItemPtr> ListViewItems;
 
 	/** ListBox for displaying items */
-	TSharedPtr<SListView<FFileTreeItemRef>> ListView;
+	TSharedPtr<SListView<FChangelistTreeItemPtr>> ListView;
 
 	/** Specify which column to sort with */
 	FName SortByColumn;
@@ -109,29 +109,3 @@ private:
 	/** The checkbox widget. */
 	TSharedPtr<SCheckBox> CheckBox;
 };
-
-class SSourceControlCheckedOutDialogListRow : public SMultiColumnTableRow<FFileTreeItemRef>
-{
-public:
-	SLATE_BEGIN_ARGS(SSourceControlCheckedOutDialogListRow)
-		: _ShowingContentVersePath(false)
-		{
-		}
-
-		SLATE_ARGUMENT(FFileTreeItemPtr, Item)
-		SLATE_ARGUMENT(bool, ShowingContentVersePath)
-
-	SLATE_END_ARGS()
-
-	/** Construct function for this widget */
-	void Construct(const FArguments& InArgs, const TSharedRef<STableViewBase>& InOwnerTableView);
-
-	/** Overridden from SMultiColumnTableRow.  Generates a widget for this column of the list row. */
-	virtual TSharedRef<SWidget> GenerateWidgetForColumn(const FName& ColumnName) override;
-
-private:
-
-	/** The item associated with this row */
-	FFileTreeItemPtr Item;
-	bool bShowingContentVersePath = false;
-};
\ No newline at end of file
diff --git a/Engine/Source/Editor/SourceControlWindows/Private/SSourceControlCommon.cpp b/Engine/Source/Editor/SourceControlWindows/Private/SSourceControlCommon.cpp
--- a/Engine/Source/Editor/SourceControlWindows/Private/SSourceControlCommon.cpp
+++ b/Engine/Source/Editor/SourceControlWindows/Private/SSourceControlCommon.cpp
@@ -168,7 +168,7 @@ static bool RefreshAssetVersePathInternal(const TArray<FAssetData>* Assets, UE::
 	return false;
 }
 
-static void RefreshAssetInformationInternal(const TArray<FAssetData>* Assets, const FString& InFilename, FString& OutAssetName, FString& OutAssetPath, UE::Core::FVersePath& OutAssetVersePath, FString& OutAssetType, FString& OutAssetTypeName, FString& OutPackageName, FColor& OutAssetTypeColor)
+static void RefreshAssetInformationInternal(const TArray<FAssetData>* Assets, const FString& InFilename, FString& OutAssetName, FString& OutAssetPath, UE::Core::FVersePath& OutAssetVersePath, FString& OutAssetType, FString& OutAssetTypeName, FColor& OutAssetTypeColor)
 {
 	// Initialize display-related members
 	FString TempAssetName = SSourceControlCommon::GetDefaultAssetName().ToString();
@@ -176,7 +176,6 @@ static void RefreshAssetInformationInternal(const TArray<FAssetData>* Assets, co
 	UE::Core::FVersePath TempAssetVersePath;
 	FString TempAssetType = SSourceControlCommon::GetDefaultAssetType().ToString();
 	FString TempAssetTypeName = SSourceControlCommon::GetDefaultAssetType().ToString();
-	FString TempPackageName;
 	FColor TempAssetColor = FColor(		// Copied from ContentBrowserCLR.cpp
 		127 + FColor::Red.R / 2,	// Desaturate the colors a bit (GB colors were too.. much)
 		127 + FColor::Red.G / 2,
@@ -227,44 +226,25 @@ static void RefreshAssetInformationInternal(const TArray<FAssetData>* Assets, co
 				TempAssetColor = AssetTypeActions->GetTypeColor();
 			}
 		}
-
-		// Beautify the package name
-		TempPackageName = TempAssetPath + "." + TempAssetName;
-	}
-	else if (bIsPackageExtension && FPackageName::TryConvertFilenameToLongPackageName(InFilename, TempPackageName))
-	{
-		// Fake asset name, asset path from the package name
-		TempAssetPath = TempPackageName;
-
-		int32 LastSlash = -1;
-		if (TempPackageName.FindLastChar('/', LastSlash))
-		{
-			TempAssetName = TempPackageName;
-			TempAssetName.RightChopInline(LastSlash + 1);
-		}
 	}
 	else
 	{
-		TempAssetName = FPaths::GetCleanFilename(InFilename);
-		TempAssetPath = InFilename;
-		TempPackageName = InFilename; // Put back original package name if the try failed
-		TempAssetType = FText::Format(SSourceControlCommon::GetDefaultUnknownAssetType(), FText::FromString(Extension.ToUpper())).ToString();
-		TempAssetTypeName = TempAssetType;
-
-		// Attempt to make package name relative to one of the project roots instead of a full absolute path
-		TArray<FSourceControlProjectInfo> CustomProjects = ISourceControlModule::Get().GetCustomProjects();
-		for (const FSourceControlProjectInfo& ProjectInfo : CustomProjects)
+		FString PackageName;
+		if (bIsPackageExtension && FPackageName::TryConvertFilenameToLongPackageName(InFilename, PackageName))
 		{
-			FStringView RelativePackageName;
-			if (FPathViews::TryMakeChildPathRelativeTo(TempPackageName, ProjectInfo.ProjectDirectory, RelativePackageName))
-			{
-				TempPackageName = FPaths::Combine(TEXT("/"), FPaths::GetBaseFilename(ProjectInfo.ProjectDirectory), RelativePackageName);
-				break;
-			}
+			// Fake asset name, asset path from the package name
+			TempAssetName = FPackageName::GetShortName(PackageName);
+			TempAssetPath = MoveTemp(PackageName);
 		}
+		else
+		{
+			TempAssetName = FPaths::GetCleanFilename(InFilename);
+			TempAssetPath = InFilename;
+			TempAssetType = FText::Format(SSourceControlCommon::GetDefaultUnknownAssetType(), FText::FromString(Extension.ToUpper())).ToString();
+			TempAssetTypeName = TempAssetType;
 
-		FPaths::MakePlatformFilename(TempAssetPath);
-		FPaths::MakePlatformFilename(TempPackageName);
+			FPaths::MakePlatformFilename(TempAssetPath);
+		}
 	}
 
 	// Finally, assign the temp variables to the member variables
@@ -274,7 +254,6 @@ static void RefreshAssetInformationInternal(const TArray<FAssetData>* Assets, co
 	OutAssetType = MoveTemp(TempAssetType);
 	OutAssetTypeName = MoveTemp(TempAssetTypeName);
 	OutAssetTypeColor = TempAssetColor;
-	OutPackageName = MoveTemp(TempPackageName);
 }
 }
 
@@ -363,13 +342,12 @@ FText FFileTreeItem::GetFileName() const
 void FFileTreeItem::RefreshAssetInformation()
 {
 	// Initialize display-related members
-	SSourceControlCommonPrivate::RefreshAssetInformationInternal(Assets.Get(), FileState->GetFilename(), AssetNameStr, AssetPathStr, AssetVersePathStruct, AssetTypeStr, AssetTypeNameStr, AssetPackageNameStr, AssetTypeColor);
+	SSourceControlCommonPrivate::RefreshAssetInformationInternal(Assets.Get(), FileState->GetFilename(), AssetNameStr, AssetPathStr, AssetVersePathStruct, AssetTypeStr, AssetTypeNameStr, AssetTypeColor);
 	AssetName = FText::FromString(AssetNameStr);
 	AssetPath = FText::FromString(AssetPathStr);
 	AssetVersePath = FText::FromString(AssetVersePathStruct.ToString());
 	AssetType = FText::FromString(AssetTypeStr);
 	AssetTypeName = FText::FromString(AssetTypeNameStr);
-	AssetPackageName = FText::FromString(AssetPackageNameStr);
 }
 
 bool FFileTreeItem::RefreshVersePath()
@@ -413,8 +391,14 @@ FText FShelvedChangelistTreeItem::GetDisplayText() const
 FOfflineFileTreeItem::FOfflineFileTreeItem(const FString& InFilename)
 	: IFileViewTreeItem(IChangelistTreeItem::OfflineFile)
 	, CheckBoxState(ECheckBoxState::Checked)
-	, Filename(InFilename)
+	, FilenameStr(InFilename)
 {
+	{
+		FString TempFilename = FilenameStr;
+		FPaths::MakePlatformFilename(TempFilename);
+		Filename = FText::FromString(MoveTemp(TempFilename));
+	}
+
 	USourceControlHelpers::GetAssetData(InFilename, Assets);
 
 	RefreshAssetInformation();
@@ -422,13 +406,12 @@ FOfflineFileTreeItem::FOfflineFileTreeItem(const FString& InFilename)
 
 void FOfflineFileTreeItem::RefreshAssetInformation()
 {
-	SSourceControlCommonPrivate::RefreshAssetInformationInternal(&Assets, Filename, AssetNameStr, AssetPathStr, AssetVersePathStruct, AssetTypeStr, AssetTypeNameStr, AssetPackageNameStr, AssetTypeColor);
+	SSourceControlCommonPrivate::RefreshAssetInformationInternal(&Assets, FilenameStr, AssetNameStr, AssetPathStr, AssetVersePathStruct, AssetTypeStr, AssetTypeNameStr, AssetTypeColor);
 	AssetName = FText::FromString(AssetNameStr);
 	AssetPath = FText::FromString(AssetPathStr);
 	AssetVersePath = FText::FromString(AssetVersePathStruct.ToString());
 	AssetType = FText::FromString(AssetTypeStr);
 	AssetTypeName = FText::FromString(AssetTypeNameStr);
-	AssetPackageName = FText::FromString(AssetPackageNameStr);
 }
 
 bool FOfflineFileTreeItem::RefreshVersePath()
diff --git a/Engine/Source/Editor/SourceControlWindows/Private/SSourceControlCommon.h b/Engine/Source/Editor/SourceControlWindows/Private/SSourceControlCommon.h
--- a/Engine/Source/Editor/SourceControlWindows/Private/SSourceControlCommon.h
+++ b/Engine/Source/Editor/SourceControlWindows/Private/SSourceControlCommon.h
@@ -107,9 +107,6 @@ struct IFileViewTreeItem : public IChangelistTreeItem
 	/** The values displayed in the 'Path' column when not displaying package names or Verse paths. */
 	virtual const FString& GetPath() const = 0;
 
-	/** The values displayed in the 'Path' column when displaying package names. */
-	virtual const FString& GetPackageName() const = 0;
-
 	/** The values displayed in the 'Path' column when displaying Verse paths. */
 	virtual const UE::Core::FVersePath& GetVersePath() const = 0;
 
@@ -243,7 +240,6 @@ struct FFileTreeItem : public IFileViewTreeItem
 	virtual int32 GetIconSortingPriority() const override;
 	virtual const FString& GetName() const override { return AssetNameStr; }
 	virtual const FString& GetPath() const override { return AssetPathStr; }
-	virtual const FString& GetPackageName() const override { return AssetPackageNameStr; }
 	virtual const UE::Core::FVersePath& GetVersePath() const override { return AssetVersePathStruct; }
 	virtual const FString& GetType() const override { return AssetTypeStr; }
 	virtual const FString& GetTypeName() const override { return AssetTypeNameStr; }
@@ -279,9 +275,6 @@ struct FFileTreeItem : public IFileViewTreeItem
 	/** Returns the user that checked out the file/asset (if any). */
 	FText GetCheckedOutByUser() const;
 
-	/** Returns the package name of the item to display */
-	FText GetAssetPackageName() const { return AssetPackageName; }
-
 	/** Returns the file name of the item in source control */
 	FText GetFileName() const;
 
@@ -323,10 +316,6 @@ private:
 	FText AssetPath;
 	FString AssetPathStr;
 
-	/** Cached package name to display */
-	FText AssetPackageName;
-	FString AssetPackageNameStr;
-
 	/** Cached Verse path to display */
 	FText AssetVersePath;
 	UE::Core::FVersePath AssetVersePathStruct;
@@ -380,15 +369,13 @@ public:
 	virtual int32 GetIconSortingPriority() const override { return 0; }
 	virtual const FString& GetName() const override { return AssetNameStr; }
 	virtual const FString& GetPath() const override { return AssetPathStr; }
-	virtual const FString& GetPackageName() const override { return AssetPackageNameStr; }
 	virtual const UE::Core::FVersePath& GetVersePath() const override { return AssetVersePathStruct; }
 	virtual const FString& GetType() const override { return AssetTypeStr; }
-	virtual const FString& GetFullPathname() const override { return Filename; }
+	virtual const FString& GetFullPathname() const override { return FilenameStr; }
 
 	virtual bool RefreshVersePath() override;
 
-	const FString& GetFilename() const { return Filename; }
-	const FText& GetDisplayPackageName() const { return AssetPackageName; }
+	const FText& GetDisplayFilename() const { return Filename; }
 	const FText& GetDisplayName() const { return AssetName; }
 	const FText& GetDisplayPath() const { return AssetPath; }
 	const FText& GetDisplayVersePath() const { return AssetVersePath; }
@@ -398,13 +385,12 @@ public:
 private:
 	ECheckBoxState CheckBoxState;
 	TArray<FAssetData> Assets;
-	FString Filename;
+	FText Filename;
+	FString FilenameStr;
 	FText AssetName;
 	FString AssetNameStr;
 	FText AssetPath;
 	FString AssetPathStr;
-	FText AssetPackageName;
-	FString AssetPackageNameStr;
 	FText AssetVersePath;
 	UE::Core::FVersePath AssetVersePathStruct;
 	FText AssetType;
diff --git a/Engine/Source/Editor/SourceControlWindows/Private/SSourceControlSubmit.cpp b/Engine/Source/Editor/SourceControlWindows/Private/SSourceControlSubmit.cpp
--- a/Engine/Source/Editor/SourceControlWindows/Private/SSourceControlSubmit.cpp
+++ b/Engine/Source/Editor/SourceControlWindows/Private/SSourceControlSubmit.cpp
@@ -813,8 +813,7 @@ TSharedRef<ITableRow> SSourceControlSubmitWidget::OnGenerateRowForList(FChangeli
 {
 	return SNew(SFileTableRow, OwnerTable)
 		.TreeItemToVisualize(SubmitItem)
-		.PathFlags((bShowingContentVersePath ? SourceControlFileViewColumn::EPathFlags::ShowingVersePath : SourceControlFileViewColumn::EPathFlags::Default) |
-			SourceControlFileViewColumn::EPathFlags::ShowingPackageName);
+		.PathFlags(bShowingContentVersePath ? SourceControlFileViewColumn::EPathFlags::ShowingVersePath : SourceControlFileViewColumn::EPathFlags::Default);
 }
 
 
@@ -857,8 +856,7 @@ void SSourceControlSubmitWidget::RequestSort()
 void SSourceControlSubmitWidget::SortTree()
 {
 	TFunction<bool(const IFileViewTreeItem&, const IFileViewTreeItem&)> SortPredicate = SourceControlFileViewColumn::GetSortPredicate(
-		SortMode, SortByColumn, (bShowingContentVersePath ? SourceControlFileViewColumn::EPathFlags::ShowingVersePath : SourceControlFileViewColumn::EPathFlags::Default) |
-		SourceControlFileViewColumn::EPathFlags::ShowingPackageName);
+		SortMode, SortByColumn, bShowingContentVersePath ? SourceControlFileViewColumn::EPathFlags::ShowingVersePath : SourceControlFileViewColumn::EPathFlags::Default);
 	if (SortPredicate)
 	{
 		Algo::SortBy(
