Skein: Fixed 'checking for assets to check-in' staying active after submit dialog opens.

Once ExpireAndFadeout() is called on a notification it will begin to fade out. However, progress on that requires its Tick() to be called which happens via the FTSTicker. Since the dialog prevents the regular engine loop from being tick'd, the dialog now ticks FTSTicker explicitly.

#rnx

diff --git a/Engine/Source/Editor/SourceControlWindows/Private/SSourceControlSubmit.cpp b/Engine/Source/Editor/SourceControlWindows/Private/SSourceControlSubmit.cpp
--- a/Engine/Source/Editor/SourceControlWindows/Private/SSourceControlSubmit.cpp
+++ b/Engine/Source/Editor/SourceControlWindows/Private/SSourceControlSubmit.cpp
@@ -676,7 +676,11 @@ FReply SSourceControlSubmitWidget::OnKeyDown( const FGeometry& MyGeometry, const
 
 void SSourceControlSubmitWidget::Tick(const FGeometry& AllottedGeometry, const double InCurrentTime, const float InDeltaTime)
 {
+	SCompoundWidget::Tick(AllottedGeometry, InCurrentTime, InDeltaTime);
+
 	ISourceControlModule::Get().Tick();
+
+	FTSTicker::GetCoreTicker().Tick(InDeltaTime);
 }
 
 ECheckBoxState SSourceControlSubmitWidget::GetToggleSelectedState() const
