SSourceControlChangelistsWidget : In "Save To Changelist" don't try to checkout a package when source control is disabled and the file does not exist.

This case can happen when source control is disabled and this operation is performed on a new asset that has not yet been saved at least once.  Previously, the call to `PackageHelper.Checkout` would fail in this case and therefore nothing would happen.  In this case it is safe to skip the checkout and allow the save to proceed.

diff --git a/Engine/Source/Editor/SourceControlWindows/Private/SSourceControlChangelists.cpp b/Engine/Source/Editor/SourceControlWindows/Private/SSourceControlChangelists.cpp
--- a/Engine/Source/Editor/SourceControlWindows/Private/SSourceControlChangelists.cpp
+++ b/Engine/Source/Editor/SourceControlWindows/Private/SSourceControlChangelists.cpp
@@ -3034,12 +3034,20 @@ TSharedPtr<SWidget> SSourceControlChangelistsWidget::OnOpenContextMenu()
 				GetSelectedFiles(SelectedFiles, SelectedFiles);
 				
 				TArray<FString> PackageNames;
+				TArray<FString> PackageNamesToCheckout;
 				for (const FString& Filename : SelectedFiles)
 				{
-					PackageNames.Add(UPackageTools::FilenameToPackageName(Filename));
+					FString PackageName = UPackageTools::FilenameToPackageName(Filename);
+					PackageNames.Add(PackageName);
+
+					// If source control is disabled and the file does not exist then don't try to checkout.
+					if (PackageHelper.UseSourceControl() || IPlatformFile::GetPlatformPhysical().FileExists(*Filename))
+					{
+						PackageNamesToCheckout.Add(PackageName);
+					}
 				}
 				
-				if (!PackageHelper.Checkout(PackageNames))
+				if (!PackageHelper.Checkout(PackageNamesToCheckout))
 				{
 					return;
 				}
