SourceControl: Made messaging in the 'SSourceControlSubmitWidget' dialog changelist agnostic.

The current messaging in this dialog is very Perforce centric ('changelist') which is not necessary.

#rnx

diff --git a/Engine/Source/Editor/SourceControlWindows/Private/SSourceControlChangelists.cpp b/Engine/Source/Editor/SourceControlWindows/Private/SSourceControlChangelists.cpp
--- a/Engine/Source/Editor/SourceControlWindows/Private/SSourceControlChangelists.cpp
+++ b/Engine/Source/Editor/SourceControlWindows/Private/SSourceControlChangelists.cpp
@@ -2457,7 +2457,7 @@ void SSourceControlChangelistsWidget::OnSubmitChangelist()
 	FText UserEditChangelistDescription = ChangelistDescriptionToSubmit;
 
 	TSharedRef<SWindow> NewWindow = SNew(SWindow)
-		.Title(NSLOCTEXT("SourceControl.ConfirmSubmit", "Title", "Confirm changelist submit"))
+		.Title(NSLOCTEXT("SourceControl.ConfirmSubmit", "Title", "Confirm Submit"))
 		.SizingRule(ESizingRule::UserSized)
 		.ClientSize(FVector2D(800, 600))
 		.SupportsMaximize(true)
diff --git a/Engine/Source/Editor/SourceControlWindows/Private/SSourceControlSubmit.cpp b/Engine/Source/Editor/SourceControlWindows/Private/SSourceControlSubmit.cpp
--- a/Engine/Source/Editor/SourceControlWindows/Private/SSourceControlSubmit.cpp
+++ b/Engine/Source/Editor/SourceControlWindows/Private/SSourceControlSubmit.cpp
@@ -259,7 +259,7 @@ void SSourceControlSubmitWidget::Construct(const FArguments& InArgs)
 	.Padding(5)
 	[
 		SNew(STextBlock)
-		.Text(NSLOCTEXT("SourceControl.SubmitPanel", "ChangeListDesc", "Changelist Description"))
+		.Text(NSLOCTEXT("SourceControl.SubmitPanel", "ChangeListDesc", "Description"))
 	];
 
 	Contents->AddSlot()
@@ -319,7 +319,7 @@ void SSourceControlSubmitWidget::Construct(const FArguments& InArgs)
 				.Padding(5)
 				[
 					SNew(SErrorText)
-					.ErrorText(NSLOCTEXT("SourceControl.SubmitPanel", "ChangeListDescWarning", "Changelist description is required to submit"))
+					.ErrorText(NSLOCTEXT("SourceControl.SubmitPanel", "ChangeListDescWarning", "Description is required to submit"))
 				]
 			]
 			+ SVerticalBox::Slot()
@@ -342,7 +342,7 @@ void SSourceControlSubmitWidget::Construct(const FArguments& InArgs)
 				.Padding(5)
 				[
 					SNew(SErrorText)
-					.ErrorText(NSLOCTEXT("SourceControl.SubmitPanel", "ChangeListErrorWarning", "Changelist errors need to be fixed to submit")) // Other errors exist and a better mechanism should be built in to display the right error. 
+					.ErrorText(NSLOCTEXT("SourceControl.SubmitPanel", "ChangeListErrorWarning", "Errors need to be fixed to submit")) // Other errors exist and a better mechanism should be built in to display the right error. 
 				]
 			]
 		];
diff --git a/Engine/Source/Editor/SourceControlWindows/Private/SourceControlWindows.cpp b/Engine/Source/Editor/SourceControlWindows/Private/SourceControlWindows.cpp
--- a/Engine/Source/Editor/SourceControlWindows/Private/SourceControlWindows.cpp
+++ b/Engine/Source/Editor/SourceControlWindows/Private/SourceControlWindows.cpp
@@ -392,7 +392,7 @@ bool FSourceControlWindows::PromptForCheckin(FCheckinResultInfo& OutResultInfo,
 	//~~~~~~~~~~~~~~~~~~~~~~~~~~~~~~~~~~~~~~~~~~
 	// Create a submit files window
 	TSharedRef<SWindow> NewWindow = SNew(SWindow)
-		.Title(NSLOCTEXT("SourceControl.SubmitWindow", "Title", "Submit Files"))
+		.Title(NSLOCTEXT("SourceControl.ConfirmSubmit", "Title", "Confirm Submit"))
 		.SizingRule(ESizingRule::UserSized)
 		.ClientSize(FVector2D(800, 600))
 		.SupportsMaximize(true)
