SourceControl: If data validation fails, the notification hyperlink that opens the MessageLog should point towards 'AssetCheck' (where validation errors appear) - not 'SourceControl'.

#rnx

diff --git a/Engine/Source/Editor/SourceControlWindows/Private/SSourceControlChangelists.cpp b/Engine/Source/Editor/SourceControlWindows/Private/SSourceControlChangelists.cpp
--- a/Engine/Source/Editor/SourceControlWindows/Private/SSourceControlChangelists.cpp
+++ b/Engine/Source/Editor/SourceControlWindows/Private/SSourceControlChangelists.cpp
@@ -2613,7 +2613,7 @@ void SSourceControlChangelistsWidget::OnValidateChangelist()
 
 	Info.ExpireDuration = 8.0f;
 	Info.HyperlinkText = LOCTEXT("SCC_Validation_ShowLog", "Show Message Log");
-	Info.Hyperlink = FSimpleDelegate::CreateLambda([]() { FMessageLog("SourceControl").Open(EMessageSeverity::Info, true); });
+	Info.Hyperlink = FSimpleDelegate::CreateLambda([]() { FMessageLog("AssetCheck").Open(EMessageSeverity::Info, true); });
 
 	TSharedPtr<SNotificationItem> Notification = FSlateNotificationManager::Get().AddNotification(Info);
 	Notification->SetCompletionState(bValidationResult ? SNotificationItem::CS_Success : SNotificationItem::CS_Fail);
