[Revision Control] Allow submit override to manage whether submit option is available
* Make shelve files check only valid if there is no override or if it defers to the base system
* Expose submit disabled reason as a tool tip

#tests Test with submit tool overriding message and not (so I saw the shelved files error properly)

diff --git a/Engine/Plugins/Experimental/SubmitToolEditor/Source/SubmitToolEditor/Private/SubmitToolEditor.cpp b/Engine/Plugins/Experimental/SubmitToolEditor/Source/SubmitToolEditor/Private/SubmitToolEditor.cpp
--- a/Engine/Plugins/Experimental/SubmitToolEditor/Source/SubmitToolEditor/Private/SubmitToolEditor.cpp
+++ b/Engine/Plugins/Experimental/SubmitToolEditor/Source/SubmitToolEditor/Private/SubmitToolEditor.cpp
@@ -96,6 +96,7 @@ void FSubmitToolEditorModule::RegisterSubmitOverrideDelegate(const USubmitToolEd
 		}
 
 		UE_LOG(LogSubmitToolEditor, Display, TEXT("Registering Submit Tool with Path:'%s' and Args: '%s'"), *InSettings->SubmitToolPath, *InSettings->SubmitToolArguments);
+		ISourceControlWindowsModule::Get().CanSubmitOverrideDelegate.BindRaw(this, &FSubmitToolEditorModule::OnCanSubmitToolOverrideCallback);
 		ISourceControlWindowsModule::Get().SubmitOverrideDelegate.BindRaw(this, &FSubmitToolEditorModule::OnSubmitToolOverrideCallback);
 	}
 }
@@ -111,14 +112,15 @@ void FSubmitToolEditorModule::UnregisterSubmitOverrideDelegate()
 	}
 }
 
-FSubmitOverrideReply FSubmitToolEditorModule::OnSubmitToolOverrideCallback(SSubmitOverrideParameters InParameters)
+FSubmitOverrideReply FSubmitToolEditorModule::OnCanSubmitToolOverrideCallback(const SSubmitOverrideParameters& InParameters, FText* ErrorMessageOut)
 {
-	ISourceControlProvider& SourceControlProvider = ISourceControlModule::Get().GetProvider();
+	return IsPerforceProvider() ? FSubmitOverrideReply::Handled : FSubmitOverrideReply::ProviderNotSupported;
+}
 
-	const FString& ProviderName = SourceControlProvider.GetName().ToString();
-	if (!ProviderName.Equals(TEXT("Perforce"), ESearchCase::IgnoreCase))
+FSubmitOverrideReply FSubmitToolEditorModule::OnSubmitToolOverrideCallback(const SSubmitOverrideParameters& InParameters)
+{
+	if (!IsPerforceProvider())
 	{
-		UE_LOG(LogSubmitToolEditor, Warning, TEXT("Current Provider is not supported for Submit Tool: '%s', Using the regular editor source control flow"), *ProviderName);
 		return FSubmitOverrideReply::ProviderNotSupported;
 	}
 
@@ -168,9 +170,10 @@ FSubmitOverrideReply FSubmitToolEditorModule::OnSubmitToolOverrideCallback(SSubm
 	{
 		UE_LOG(LogSubmitToolEditor, Error, TEXT("The path is invalid: file does not exist - '%s'"), *NormalizedPath);		
 		return Settings->bForceSubmitTool ? FSubmitOverrideReply::Error : FSubmitOverrideReply::ProviderNotSupported;
-
 	}
 
+	ISourceControlProvider& SourceControlProvider = ISourceControlModule::Get().GetProvider();
+
 	if (InParameters.ToSubmit.HasSubtype<FString>())
 	{
 		return InvokeSubmitTool(SourceControlProvider, NormalizedPath, Settings->SubmitToolArguments, InParameters.Description, InParameters.ToSubmit.GetSubtype<FString>());
@@ -263,6 +266,19 @@ FSubmitOverrideReply FSubmitToolEditorModule::InvokeSubmitTool(ISourceControlPro
 	return FSubmitOverrideReply::Handled;
 }
 
+bool FSubmitToolEditorModule::IsPerforceProvider() const
+{
+	ISourceControlProvider& SourceControlProvider = ISourceControlModule::Get().GetProvider();
+
+	if (SourceControlProvider.GetName() == TEXT("Perforce"))
+	{
+		return true;
+	}
+
+	UE_LOG(LogSubmitToolEditor, Warning, TEXT("Current Provider is not supported for Submit Tool: '%s', Using the regular editor source control flow"), *SourceControlProvider.GetName().ToString());
+	return false;
+}
+
 bool FSubmitToolEditorModule::GetPerforceParameter(ISourceControlProvider& InProvider, FString& OutPort, FString& OutUser, FString& OutClient, FString& OutWorkspacePath)
 {
 	TMap<ISourceControlProvider::EStatus, FString> Status = InProvider.GetStatus();
diff --git a/Engine/Plugins/Experimental/SubmitToolEditor/Source/SubmitToolEditor/Private/SubmitToolEditor.h b/Engine/Plugins/Experimental/SubmitToolEditor/Source/SubmitToolEditor/Private/SubmitToolEditor.h
--- a/Engine/Plugins/Experimental/SubmitToolEditor/Source/SubmitToolEditor/Private/SubmitToolEditor.h
+++ b/Engine/Plugins/Experimental/SubmitToolEditor/Source/SubmitToolEditor/Private/SubmitToolEditor.h
@@ -34,11 +34,13 @@ public:
 	
 private:
 
-	FSubmitOverrideReply OnSubmitToolOverrideCallback(SSubmitOverrideParameters InParameters);
+	FSubmitOverrideReply OnCanSubmitToolOverrideCallback(const SSubmitOverrideParameters& InParameters, FText* ErrorMessageOut);
+	FSubmitOverrideReply OnSubmitToolOverrideCallback(const SSubmitOverrideParameters& InParameters);
 	FSubmitOverrideReply InvokeSubmitTool(ISourceControlProvider& InProvider, const FString& InPath, const FString& InArgs, const FString& InDescription, const FString& InIdentifier);
 	FSubmitOverrideReply InvokeSubmitTool(ISourceControlProvider& InProvider, const FString& InPath, const FString& InArgs, const FString& InDescription, const TArray<FString>& InFiles);
 	FSubmitOverrideReply InvokeSubmitTool(ISourceControlProvider& InProvider, const FString& InPath, const FString& InArgs, const FString& InIdentifier);
 
+	bool IsPerforceProvider() const;
 	bool GetPerforceParameter(ISourceControlProvider& InProvider, FString& OutPort, FString& OutUser, FString& OutClient, FString& OutWorkspacePath);
 	bool EditChangelistDescription(ISourceControlProvider& InProvider, const FString& InNewChangelistDescription, const FString& InIdentifier, FString& OutIdentifier);
 	void SaveChangelistDescription(ISourceControlProvider& InProvider, const FString& InNewChangelistDescription, const FString& InIdentifier);
diff --git a/Engine/Source/Editor/SourceControlWindows/Private/SSourceControlChangelists.cpp b/Engine/Source/Editor/SourceControlWindows/Private/SSourceControlChangelists.cpp
--- a/Engine/Source/Editor/SourceControlWindows/Private/SSourceControlChangelists.cpp
+++ b/Engine/Source/Editor/SourceControlWindows/Private/SSourceControlChangelists.cpp
@@ -1538,7 +1538,7 @@ void SSourceControlChangelistsWidget::OnUncontrolledChangelistStateChanged()
 	OnRefreshUI(ERefreshFlags::UncontrolledChangelists);
 }
 
-FSourceControlChangelistStatePtr SSourceControlChangelistsWidget::GetCurrentChangelistState()
+FSourceControlChangelistStatePtr SSourceControlChangelistsWidget::GetCurrentChangelistState() const
 {
 	if (!ChangelistTreeView)
 	{
@@ -2023,12 +2023,12 @@ void SSourceControlChangelistsWidget::OnDeleteChangelist()
 		}));
 }
 
-bool SSourceControlChangelistsWidget::CanDeleteChangelist()
+bool SSourceControlChangelistsWidget::CanDeleteChangelist() const
 {
 	return CanDeleteChangelist(/*OutFailureMessage*/nullptr);
 }
 
-bool SSourceControlChangelistsWidget::CanDeleteChangelist(FText* OutFailureMessage)
+bool SSourceControlChangelistsWidget::CanDeleteChangelist(FText* OutFailureMessage) const
 {
 	FSourceControlChangelistStatePtr ChangelistState = GetCurrentChangelistState();
 
@@ -2604,12 +2604,22 @@ void SSourceControlChangelistsWidget::OnSubmitChangelist()
 	}
 }
 
-bool SSourceControlChangelistsWidget::CanSubmitChangelist()
+FText SSourceControlChangelistsWidget::GetCanSubmitChangelistToolTip() const
+{
+	FText Error;
+	if (!CanSubmitChangelist(&Error))
+	{
+		return Error;
+	}
+	return LOCTEXT("SourceControl_SubmitChangeslit_Tooltip", "Submits a changelist");
+}
+
+bool SSourceControlChangelistsWidget::CanSubmitChangelist() const
 {
 	return CanSubmitChangelist(/*OutFailureMessage*/nullptr);
 }
 
-bool SSourceControlChangelistsWidget::CanSubmitChangelist(FText* OutFailureMessage)
+bool SSourceControlChangelistsWidget::CanSubmitChangelist(FText* OutFailureMessage) const
 {
 	FSourceControlChangelistStatePtr Changelist = GetCurrentChangelistState();
 
@@ -2629,7 +2639,33 @@ bool SSourceControlChangelistsWidget::CanSubmitChangelist(FText* OutFailureMessa
 		}
 		return false;
 	}
-	else if (Changelist->GetShelvedFilesStatesNum() > 0)
+	else if (ISourceControlWindowsModule::Get().CanSubmitOverrideDelegate.IsBound())
+	{
+		// save the changelist description from the widget to perforce
+		const FString Identifier = Changelist->GetChangelist()->GetIdentifier();
+
+		SSubmitOverrideParameters SubmitOverrideParameters;
+		SubmitOverrideParameters.Description = Changelist->GetDescriptionText().ToString();
+		SubmitOverrideParameters.ToSubmit.SetSubtype<FString>(Identifier);
+
+		FSubmitOverrideReply SubmitOverrideReply = ISourceControlWindowsModule::Get().CanSubmitOverrideDelegate.Execute(SubmitOverrideParameters, OutFailureMessage);
+
+		switch (SubmitOverrideReply)
+		{
+		case FSubmitOverrideReply::Handled:
+			return true;
+		case FSubmitOverrideReply::Error:
+			return false;
+
+			//////////////////////////////////////////////////////////
+		case FSubmitOverrideReply::ProviderNotSupported:
+		default:
+			// continue default flow
+			break;
+		}
+	}
+
+	if (Changelist->GetShelvedFilesStatesNum() > 0)
 	{
 		if (OutFailureMessage)
 		{
@@ -2672,7 +2708,7 @@ void SSourceControlChangelistsWidget::OnValidateChangelist()
 	Notification->SetCompletionState(bValidationResult ? SNotificationItem::CS_Success : SNotificationItem::CS_Fail);
 }
 
-bool SSourceControlChangelistsWidget::CanValidateChangelist()
+bool SSourceControlChangelistsWidget::CanValidateChangelist() const
 {
 	FSourceControlChangelistStatePtr Changelist = GetCurrentChangelistState();
 	return Changelist != nullptr && Changelist->GetFilesStatesNum() > 0;
@@ -3152,7 +3188,7 @@ TSharedPtr<SWidget> SSourceControlChangelistsWidget::OnOpenContextMenu()
 		Section.AddMenuEntry(
 			"SubmitChangelist",
 			LOCTEXT("SourceControl_SubmitChangelist", "Submit Changelist..."),
-			LOCTEXT("SourceControl_SubmitChangeslit_Tooltip", "Submits a changelist"),
+			{ this, &SSourceControlChangelistsWidget::GetCanSubmitChangelistToolTip },
 			FSlateIcon(),
 			FUIAction(
 				FExecuteAction::CreateSP(this, &SSourceControlChangelistsWidget::OnSubmitChangelist),
diff --git a/Engine/Source/Editor/SourceControlWindows/Private/SSourceControlChangelists.h b/Engine/Source/Editor/SourceControlWindows/Private/SSourceControlChangelists.h
--- a/Engine/Source/Editor/SourceControlWindows/Private/SSourceControlChangelists.h
+++ b/Engine/Source/Editor/SourceControlWindows/Private/SSourceControlChangelists.h
@@ -104,7 +104,7 @@ private:
 	TSharedPtr<SWidget> OnOpenContextMenu();
 
 	/** Returns the currently selected changelist state ptr or null in invalid cases */
-	FSourceControlChangelistStatePtr GetCurrentChangelistState();
+	FSourceControlChangelistStatePtr GetCurrentChangelistState() const;
 	FUncontrolledChangelistStatePtr GetCurrentUncontrolledChangelistState() const;
 	FSourceControlChangelistPtr GetCurrentChangelist();
 	TOptional<FUncontrolledChangelist> GetCurrentUncontrolledChangelist() const;
@@ -139,14 +139,15 @@ private:
 	/** Changelist operations */
 	void OnNewChangelist();
 	void OnDeleteChangelist();
-	bool CanDeleteChangelist();
-	bool CanDeleteChangelist(FText* OutFailureMessage);
+	bool CanDeleteChangelist() const;
+	bool CanDeleteChangelist(FText* OutFailureMessage) const;
 	void OnEditChangelist();
 	void OnSubmitChangelist();
-	bool CanSubmitChangelist();
-	bool CanSubmitChangelist(FText* OutFailureMessage);
+	FText GetCanSubmitChangelistToolTip() const;
+	bool CanSubmitChangelist() const;
+	bool CanSubmitChangelist(FText* OutFailureMessage) const;
 	void OnValidateChangelist();
-	bool CanValidateChangelist();
+	bool CanValidateChangelist() const;
 
 	/** Uncontrolled Changelist operations */
 	void OnNewUncontrolledChangelist();
diff --git a/Engine/Source/Editor/SourceControlWindows/Public/ISourceControlWindowsModule.h b/Engine/Source/Editor/SourceControlWindows/Public/ISourceControlWindowsModule.h
--- a/Engine/Source/Editor/SourceControlWindows/Public/ISourceControlWindowsModule.h
+++ b/Engine/Source/Editor/SourceControlWindows/Public/ISourceControlWindowsModule.h
@@ -26,7 +26,8 @@ struct SSubmitOverrideParameters
 	FString Description;
 };
 
-DECLARE_DELEGATE_RetVal_OneParam(FSubmitOverrideReply, FSubmitOverrideDelegate, SSubmitOverrideParameters /*InParameters*/);
+DECLARE_DELEGATE_RetVal_TwoParams(FSubmitOverrideReply, FCanSubmitOverrideDelegate, const SSubmitOverrideParameters& InParameters /*InParameters*/, FText* ErrorMessageOut);
+DECLARE_DELEGATE_RetVal_OneParam(FSubmitOverrideReply, FSubmitOverrideDelegate, const SSubmitOverrideParameters& InParameters /*InParameters*/);
 
 class ISourceControlWindowsModule : public IModuleInterface
 {
@@ -58,5 +59,6 @@ public:
 	DECLARE_EVENT_OneParam(ISourceControlWindowsModule, FChangelistFileDoubleClickedEvent, const FString&);
 	virtual FChangelistFileDoubleClickedEvent& OnChangelistFileDoubleClicked() = 0;
 
+	FCanSubmitOverrideDelegate CanSubmitOverrideDelegate;
 	FSubmitOverrideDelegate SubmitOverrideDelegate;
 };
\ No newline at end of file
