Source control: Added functions to get the names of source control providers

- Replaced usage of hard-coded strings for the source control provider. Strings were replaced with a reference to functions in SourceControlProviders.h

diff --git a/Engine/Source/Developer/SourceControl/Public/SourceControlProviders.h b/Engine/Source/Developer/SourceControl/Public/SourceControlProviders.h
new file mode 100644
--- /dev/null
+++ b/Engine/Source/Developer/SourceControl/Public/SourceControlProviders.h
@@ -0,0 +1,32 @@
+// Copyright Epic Games, Inc. All Rights Reserved.
+
+#pragma once
+
+#include "CoreMinimal.h"
+
+namespace SourceControlProviders
+{
+	/** Returns the default SourceControl provider name, which is the one indicating 'disabled'. */
+	static const TCHAR* GetDefaultProviderName()
+	{
+		return TEXT("None");
+	}
+
+	/** Returns the URC SourceControl provider name */
+	static const TCHAR* GetUrcProviderName()
+	{
+		return TEXT("Unreal Revision Control");
+	}	
+	
+	/** Returns the old Skein SourceControl provider name */
+    static const TCHAR* GetLegacySkeinProviderName()
+	{
+		return TEXT("Skein");
+	}
+
+	/** Returns the Perforce SourceControl provider name */
+	static const TCHAR* GetPerforceProviderName()
+	{
+		return TEXT("Perforce");
+	}
+}
\ No newline at end of file
diff --git a/Engine/Source/Editor/MainFrame/Private/Menus/MainMenu.cpp b/Engine/Source/Editor/MainFrame/Private/Menus/MainMenu.cpp
--- a/Engine/Source/Editor/MainFrame/Private/Menus/MainMenu.cpp
+++ b/Engine/Source/Editor/MainFrame/Private/Menus/MainMenu.cpp
@@ -7,6 +7,7 @@
 #include "Interfaces/IMainFrameModule.h"
 #include "DesktopPlatformModule.h"
 #include "ISourceControlModule.h"
+#include "SourceControlProviders.h"
 #include "IUndoHistoryEditorModule.h"
 #include "Toolkits/AssetEditorToolkit.h"
 #include "Toolkits/GlobalEditorCommonCommands.h"
@@ -661,7 +662,7 @@ void FMainMenu::RegisterToolsMenu()
 
 		// Access to the legacy 'Submit Files' dialog. Git / Subversion require this.
 		if (SourceControlModule.IsEnabled() && !SourceControlModule.GetProvider().UsesChangelists() &&
-			SourceControlModule.GetProvider().GetName() != TEXT("Unreal Revision Control"))
+			SourceControlModule.GetProvider().GetName() != SourceControlProviders::GetUrcProviderName())
 		{
 			InSection.AddMenuEntry(
 				FMainFrameCommands::Get().SubmitContent,
@@ -691,7 +692,7 @@ void FMainMenu::RegisterToolsMenu()
 	{
 		ISourceControlModule& SourceControlModule = ISourceControlModule::Get();
 
-		if (SourceControlModule.GetProvider().GetName() == TEXT("Unreal Revision Control"))
+		if (SourceControlModule.GetProvider().GetName() == SourceControlProviders::GetUrcProviderName())
 		{
 			InSection.AddMenuEntry(
 				FMainFrameCommands::Get().ViewSnapshotHistory,
diff --git a/Engine/Source/Editor/SourceControlWindows/Private/SourceControlWindowsModule.cpp b/Engine/Source/Editor/SourceControlWindows/Private/SourceControlWindowsModule.cpp
--- a/Engine/Source/Editor/SourceControlWindows/Private/SourceControlWindowsModule.cpp
+++ b/Engine/Source/Editor/SourceControlWindows/Private/SourceControlWindowsModule.cpp
@@ -4,6 +4,7 @@
 #include "Modules/ModuleManager.h"
 #include "ISourceControlWindowsModule.h"
 #include "ISourceControlModule.h"
+#include "SourceControlProviders.h"
 
 #include "Widgets/Docking/SDockTab.h"
 #include "Textures/SlateIcon.h"
@@ -145,7 +146,7 @@ bool FSourceControlWindowsModule::CanShowSnapshotHistoryTab() const
 {
 	ISourceControlModule& SourceControlModule = ISourceControlModule::Get();
 
-	return (SourceControlModule.IsEnabled() && SourceControlModule.GetProvider().IsAvailable() && SourceControlModule.GetProvider().GetName() == TEXT("Unreal Revision Control"));
+	return (SourceControlModule.IsEnabled() && SourceControlModule.GetProvider().IsAvailable() && SourceControlModule.GetProvider().GetName() == SourceControlProviders::GetUrcProviderName());
 }
 
 void FSourceControlWindowsModule::ShowConflictResolutionTab()
@@ -160,7 +161,7 @@ bool FSourceControlWindowsModule::CanShowConflictResolutionTab() const
 {
 	ISourceControlModule& SourceControlModule = ISourceControlModule::Get();
 
-	return (SourceControlModule.IsEnabled() && SourceControlModule.GetProvider().IsAvailable() && SourceControlModule.GetProvider().GetName() == TEXT("Unreal Revision Control"));
+	return (SourceControlModule.IsEnabled() && SourceControlModule.GetProvider().IsAvailable() && SourceControlModule.GetProvider().GetName() == SourceControlProviders::GetUrcProviderName());
 }
 
 void FSourceControlWindowsModule::SelectFiles(const TArray<FString>& Filenames)
diff --git a/Engine/Source/Editor/StatusBar/Private/SourceControlMenuHelpers.cpp b/Engine/Source/Editor/StatusBar/Private/SourceControlMenuHelpers.cpp
--- a/Engine/Source/Editor/StatusBar/Private/SourceControlMenuHelpers.cpp
+++ b/Engine/Source/Editor/StatusBar/Private/SourceControlMenuHelpers.cpp
@@ -5,6 +5,7 @@
 #include "ISourceControlProvider.h"
 #include "ISourceControlModule.h"
 #include "ISourceControlWindowsModule.h"
+#include "SourceControlProviders.h"
 #include "SourceControlWindows.h"
 #include "UnsavedAssetsTrackerModule.h"
 #include "FileHelpers.h"
@@ -136,7 +137,7 @@ bool FSourceControlCommands::SubmitContent_IsVisible()
 		return false;
 	}
 
-	if (ISourceControlModule::Get().GetProvider().GetName() == TEXT("Unreal Revision Control"))
+	if (ISourceControlModule::Get().GetProvider().GetName() == SourceControlProviders::GetUrcProviderName())
 	{
 		return false;
 	}
@@ -151,7 +152,7 @@ bool FSourceControlCommands::ViewSnapshotHistory_CanExecute()
 
 bool FSourceControlCommands::ViewSnapshotHistory_IsVisible()
 {
-	return ISourceControlModule::Get().GetProvider().GetName() == TEXT("Unreal Revision Control");
+	return ISourceControlModule::Get().GetProvider().GetName() == SourceControlProviders::GetUrcProviderName();
 }
 
 void FSourceControlCommands::ViewChangelists_Clicked()
diff --git a/Engine/Source/Editor/UnrealEd/Private/FileHelpers.cpp b/Engine/Source/Editor/UnrealEd/Private/FileHelpers.cpp
--- a/Engine/Source/Editor/UnrealEd/Private/FileHelpers.cpp
+++ b/Engine/Source/Editor/UnrealEd/Private/FileHelpers.cpp
@@ -91,6 +91,7 @@
 #include "PackageSourceControlHelper.h"
 #include "InterchangeManager.h"
 #include "SourceControlHelpers.h"
+#include "SourceControlProviders.h"
 #include "InterchangeProjectSettings.h"
 
 #include UE_INLINE_GENERATED_CPP_BY_NAME(FileHelpers)
@@ -468,7 +469,7 @@ namespace FileDialogHelpers
 */
 static bool UseAlternateCheckoutWorkflow()
 {
-	if (ISourceControlModule::Get().GetProvider().GetName() == TEXT("Unreal Revision Control"))
+	if (ISourceControlModule::Get().GetProvider().GetName() == SourceControlProviders::GetUrcProviderName())
 	{
 		if (const UEditorLoadingSavingSettings* Settings = GetDefault<UEditorLoadingSavingSettings>())
 		{
