Skein / URC2: Add FScopedSlowTasks for long running source control operations (sync, revert) that could potentially trigger an expensive world reload.

#rnx

diff --git a/Engine/Source/Editor/SceneOutliner/Private/SceneOutlinerSCCHandler.cpp b/Engine/Source/Editor/SceneOutliner/Private/SceneOutlinerSCCHandler.cpp
--- a/Engine/Source/Editor/SceneOutliner/Private/SceneOutlinerSCCHandler.cpp
+++ b/Engine/Source/Editor/SceneOutliner/Private/SceneOutlinerSCCHandler.cpp
@@ -452,6 +452,10 @@ void FSceneOutlinerSCCHandler::ExecuteSCCRevert()
 	if (PackagesToRevert.Num() > 0)
 	{
 		FBookmarkScoped BookmarkScoped;
+		FScopedSlowTask SlowTaskScoped(1.f, NSLOCTEXT("SourceControl", "RevertAll", "Reverting change(s)..."));
+		SlowTaskScoped.MakeDialog();
+		SlowTaskScoped.EnterProgressFrame(1.f);
+
 		SourceControlHelpers::RevertAndReloadPackages(PackagesToRevert, /*bRevertAll=*/false, /*bReloadWorld=*/true);
 	}
 }
diff --git a/Engine/Source/Editor/SourceControlWindows/Private/SSourceControlSubmit.cpp b/Engine/Source/Editor/SourceControlWindows/Private/SSourceControlSubmit.cpp
--- a/Engine/Source/Editor/SourceControlWindows/Private/SSourceControlSubmit.cpp
+++ b/Engine/Source/Editor/SourceControlWindows/Private/SSourceControlSubmit.cpp
@@ -675,6 +675,10 @@ void SSourceControlSubmitWidget::OnRevert()
 
 	{
 		FBookmarkScoped BookmarkScoped;
+		FScopedSlowTask SlowTaskScoped(1.f, NSLOCTEXT("SourceControl", "RevertAll", "Reverting change(s)..."));
+		SlowTaskScoped.MakeDialog();
+		SlowTaskScoped.EnterProgressFrame(1.f);
+
 		bool bAnyReverted = false;
 		if (PackagesToRevert.Num() > 0)
 		{
diff --git a/Engine/Source/Editor/SourceControlWindows/Private/SourceControlWindows.cpp b/Engine/Source/Editor/SourceControlWindows/Private/SourceControlWindows.cpp
--- a/Engine/Source/Editor/SourceControlWindows/Private/SourceControlWindows.cpp
+++ b/Engine/Source/Editor/SourceControlWindows/Private/SourceControlWindows.cpp
@@ -14,6 +14,7 @@
 #include "Logging/MessageLog.h"
 #include "Logging/TokenizedMessage.h"
 #include "Misc/MessageDialog.h"
+#include "Misc/ScopedSlowTask.h"
 #include "Widgets/Notifications/SNotificationList.h"
 #include "SourceControlSettings.h"
 #include "Bookmarks/BookmarkScoped.h"
@@ -440,7 +441,10 @@ bool FSourceControlWindows::PromptForCheckin(FCheckinResultInfo& OutResultInfo,
 		const bool bSyncNeeded = FSourceControlWindows::CanSyncLatest();
 		if (bSyncNeeded)
 		{
-			FBookmarkScoped BookmarkScoped; // Preserve viewport camera orientation.
+			FBookmarkScoped BookmarkScoped;
+			FScopedSlowTask SlowTaskScoped(1.f, NSLOCTEXT("SourceControl", "RevertAll", "Reverting change(s)..."));
+			SlowTaskScoped.MakeDialog();
+			SlowTaskScoped.EnterProgressFrame(1.f);
 
 			if (!FSourceControlWindows::SyncLatest())
 			{
diff --git a/Engine/Source/Editor/StatusBar/Private/SourceControlMenuHelpers.cpp b/Engine/Source/Editor/StatusBar/Private/SourceControlMenuHelpers.cpp
--- a/Engine/Source/Editor/StatusBar/Private/SourceControlMenuHelpers.cpp
+++ b/Engine/Source/Editor/StatusBar/Private/SourceControlMenuHelpers.cpp
@@ -23,6 +23,7 @@
 #include "EditorModeManager.h"
 #include "Misc/ConfigCacheIni.h"
 #include "Misc/MessageDialog.h"
+#include "Misc/ScopedSlowTask.h"
 #include "RevisionControlStyle/RevisionControlStyle.h"
 #include "Bookmarks/BookmarkScoped.h"
 #include "Styling/StyleColors.h"
@@ -237,6 +238,10 @@ void FSourceControlCommands::RevertAllModifiedFiles_Clicked()
 		FEditorFileUtils::SaveDirtyPackages(bPromptUserToSave, bSaveMapPackages, bSaveContentPackages, bFastSave, bNotifyNoPackagesSaved, bCanBeDeclined);
 
 		FBookmarkScoped BookmarkScoped;
+		FScopedSlowTask SlowTaskScoped(1.f, NSLOCTEXT("SourceControl", "RevertAll", "Reverting change(s)..."));
+		SlowTaskScoped.MakeDialog();
+		SlowTaskScoped.EnterProgressFrame(1.f);
+
 		FSourceControlWindows::RevertAllChangesAndReloadWorld();
 	}
 }
