SourceControl: Fixed copy/paste message error introduced in #44739569.

#rnx

diff --git a/Engine/Source/Editor/SourceControlWindows/Private/SourceControlWindows.cpp b/Engine/Source/Editor/SourceControlWindows/Private/SourceControlWindows.cpp
--- a/Engine/Source/Editor/SourceControlWindows/Private/SourceControlWindows.cpp
+++ b/Engine/Source/Editor/SourceControlWindows/Private/SourceControlWindows.cpp
@@ -442,7 +442,7 @@ bool FSourceControlWindows::PromptForCheckin(FCheckinResultInfo& OutResultInfo,
 		if (bSyncNeeded)
 		{
 			FBookmarkScoped BookmarkScoped;
-			FScopedSlowTask SlowTaskScoped(1.f, NSLOCTEXT("SourceControl", "RevertAll", "Reverting change(s)..."));
+			FScopedSlowTask SlowTaskScoped(1.f, NSLOCTEXT("SourceControl", "SyncLatest", "Syncing latest change(s)..."));
 			SlowTaskScoped.MakeDialog();
 			SlowTaskScoped.EnterProgressFrame(1.f);
 
