SourceControl: SSourceControlSubmitWidget::OnCreateContextMenu() now uses UToolMenus, so the context menu can be extended by external systems.

#rnx

diff --git a/Engine/Source/Editor/SourceControlWindows/Private/SSourceControlCommon.cpp b/Engine/Source/Editor/SourceControlWindows/Private/SSourceControlCommon.cpp
--- a/Engine/Source/Editor/SourceControlWindows/Private/SSourceControlCommon.cpp
+++ b/Engine/Source/Editor/SourceControlWindows/Private/SSourceControlCommon.cpp
@@ -437,8 +437,19 @@ TSharedRef<SWidget> GetSCCStatusWidget(FSourceControlStateRef InFileState)
 			.WidthOverride(SizeOverride)
 			.HeightOverride(SizeOverride)
 			[
-				SNew(SLayeredImage, InFileState->GetIcon())
-				.ToolTipText(InFileState->GetDisplayTooltip())
+				SNew(SImage)
+				.Image_Lambda(
+					[InFileState]() -> const FSlateBrush*
+					{
+						return InFileState->GetIcon().GetIcon();
+					}
+				)
+				.ToolTipText_Lambda(
+					[InFileState]() -> FText
+					{
+						return InFileState->GetDisplayTooltip();
+					}
+				)
 			]
 		];
 }
diff --git a/Engine/Source/Editor/SourceControlWindows/Private/SSourceControlSubmit.cpp b/Engine/Source/Editor/SourceControlWindows/Private/SSourceControlSubmit.cpp
--- a/Engine/Source/Editor/SourceControlWindows/Private/SSourceControlSubmit.cpp
+++ b/Engine/Source/Editor/SourceControlWindows/Private/SSourceControlSubmit.cpp
@@ -8,6 +8,7 @@
 #include "ISourceControlProvider.h"
 #include "ISourceControlModule.h"
 #include "SourceControlHelpers.h"
+#include "SourceControlWindows.h"
 #include "SSourceControlCommon.h"
 #include "SSourceControlChangelistRows.h"
 #include "Modules/ModuleManager.h"
@@ -32,6 +33,8 @@
 #include "Bookmarks/BookmarkScoped.h"
 #include "HAL/IConsoleManager.h"
 #include "Algo/AllOf.h"
+#include "ToolMenu.h"
+#include "ToolMenus.h"
 
 #if SOURCE_CONTROL_WITH_SLATE
 
@@ -537,39 +540,70 @@ void SSourceControlSubmitWidget::Construct(const FArguments& InArgs)
 /** Corvus: Called to create a context menu when right-clicking on an item */
 TSharedPtr<SWidget> SSourceControlSubmitWidget::OnCreateContextMenu()
 {
-	FMenuBuilder MenuBuilder(true, NULL);
+	static const FName MenuName = "SourceControl.SubmitContextMenu";
 
-	MenuBuilder.BeginSection("Source Control", NSLOCTEXT("SourceControl.SubmitWindow.Menu", "SourceControlSectionHeader", "Revision Control"));
+	// Register menu
+	UToolMenus* ToolMenus = UToolMenus::Get();
+	if (!ToolMenus->IsMenuRegistered(MenuName))
 	{
-		if (SSourceControlSubmitWidget::CanDiffAgainstDepot())
-		{
-			MenuBuilder.AddMenuEntry(
-				NSLOCTEXT("SourceControl.SubmitWindow.Menu", "DiffAgainstDepot", "Diff Against Depot"),
-				NSLOCTEXT("SourceControl.SubmitWindow.Menu", "DiffAgainstDepotTooltip", "Look at differences between your version of the asset and that in revision control."),
-				FSlateIcon(FAppStyle::GetAppStyleSetName(), "SourceControl.Actions.Diff"),
-				FUIAction(
-					FExecuteAction::CreateSP(this, &SSourceControlSubmitWidget::OnDiffAgainstDepot),
-					FCanExecuteAction::CreateSP(this, &SSourceControlSubmitWidget::CanDiffAgainstDepot)
-				)
-			);
-		}
+		UToolMenu* RegisteredMenu = ToolMenus->RegisterMenu(MenuName);
+		// Add section so it can be used as insert position for menu extensions
+		RegisteredMenu->AddSection("Source Control");
+	}
+
+	// Create context
+	USourceControlSubmitWidgetContext* SourceControlSubmitWidgetContext = NewObject<USourceControlSubmitWidgetContext>();
+	SourceControlSubmitWidgetContext->SubmitWidget = SharedThis(this);
 
-		if (AllowRevert())
+	const auto& SelectedItems = ListView->GetSelectedItems();
+	for (const FChangelistTreeItemPtr& SelectedItem : SelectedItems)
+	{
+		const FFileTreeItem* FileItem = GetFileItem(SelectedItem);
+		if (FileItem != nullptr)
 		{
-			MenuBuilder.AddMenuEntry(
-				NSLOCTEXT("SourceControl.SubmitWindow.Menu", "Revert", "Revert"),
-				NSLOCTEXT("SourceControl.SubmitWindow.Menu", "RevertTooltip", "Revert the selected assets to their original state from revision control."),
-				FSlateIcon(FAppStyle::GetAppStyleSetName(), "SourceControl.Actions.Revert"),
-				FUIAction(
-					FExecuteAction::CreateSP(this, &SSourceControlSubmitWidget::OnRevert),
-					FCanExecuteAction::CreateSP(this, &SSourceControlSubmitWidget::CanRevert)
-				)
-			);
+			USourceControlSubmitWidgetContext::SelectedItem ItemInfo;
+			ItemInfo.FileName = FileItem->GetFileName().ToString();
+
+			SourceControlSubmitWidgetContext->GetSelectedItems().Add(ItemInfo);
 		}
 	}
-	MenuBuilder.EndSection();
 
-	return MenuBuilder.MakeWidget();
+	FToolMenuContext Context;
+	Context.AddObject(SourceControlSubmitWidgetContext);
+
+	// Generate menu
+	UToolMenu* Menu = ToolMenus->GenerateMenu(MenuName, Context);
+
+	FToolMenuSection& Section = *Menu->FindSection("Source Control");
+	if (SSourceControlSubmitWidget::CanDiffAgainstDepot())
+	{
+		Section.AddMenuEntry(
+			NAME_None,
+			NSLOCTEXT("SourceControl.SubmitWindow.Menu", "DiffAgainstDepot", "Diff Against Depot"),
+			NSLOCTEXT("SourceControl.SubmitWindow.Menu", "DiffAgainstDepotTooltip", "Look at differences between your version of the asset and that in revision control."),
+			FSlateIcon(FAppStyle::GetAppStyleSetName(), "SourceControl.Actions.Diff"),
+			FUIAction(
+				FExecuteAction::CreateSP(this, &SSourceControlSubmitWidget::OnDiffAgainstDepot),
+				FCanExecuteAction::CreateSP(this, &SSourceControlSubmitWidget::CanDiffAgainstDepot)
+			)
+		);
+	}
+
+	if (AllowRevert())
+	{
+		Section.AddMenuEntry(
+			NAME_None,
+			NSLOCTEXT("SourceControl.SubmitWindow.Menu", "Revert", "Revert"),
+			NSLOCTEXT("SourceControl.SubmitWindow.Menu", "RevertTooltip", "Revert the selected assets to their original state from revision control."),
+			FSlateIcon(FAppStyle::GetAppStyleSetName(), "SourceControl.Actions.Revert"),
+			FUIAction(
+				FExecuteAction::CreateSP(this, &SSourceControlSubmitWidget::OnRevert),
+				FCanExecuteAction::CreateSP(this, &SSourceControlSubmitWidget::CanRevert)
+			)
+		);
+	}
+
+	return ToolMenus->GenerateWidget(Menu);
 }
 
 bool SSourceControlSubmitWidget::CanDiffAgainstDepot() const
diff --git a/Engine/Source/Editor/SourceControlWindows/Public/SourceControlWindows.h b/Engine/Source/Editor/SourceControlWindows/Public/SourceControlWindows.h
--- a/Engine/Source/Editor/SourceControlWindows/Public/SourceControlWindows.h
+++ b/Engine/Source/Editor/SourceControlWindows/Public/SourceControlWindows.h
@@ -38,6 +38,28 @@ private:
 	TArray<SelectedItem> SelectedItems;
 };
 
+/** Context object for SSourceControlSubmitWidget's right-click context menu to get info about which submit widget is trying to generate the menu */
+UCLASS(MinimalAPI)
+class USourceControlSubmitWidgetContext : public UObject
+{
+	GENERATED_BODY()
+
+public:
+	struct SelectedItem
+	{
+		FString FileName;
+	};
+
+	TWeakPtr<class SSourceControlSubmitWidget> SubmitWidget;
+
+	TArray<SelectedItem>& GetSelectedItems()
+	{
+		return SelectedItems;
+	}
+
+private:
+	TArray<SelectedItem> SelectedItems;
+};
 
 /** Info supplied as argument to delegate FSourceControlWindowsOnCheckInComplete called by FSourceControlWindows::ChoosePackagesToCheckIn() and optional argument to PromptForCheckin(). */
 struct FCheckinResultInfo
