Save Assets button in revision windows

Problem
No quick way to save all unsaved assets in a selected changelist.
Users had to manually save each asset or all of them, slowing down workflow.

Solution
Added "Save Assets" button to the context menu on changelists and items.
Button dynamically enables only if unsaved assets exist in the changelist.
Only saves assets within the selected changelist (not all assets).
No icon used, consistent with other context menu entries.

Notes
Easy to validate: modify assets in a CL, right-click, confirm Save Assets button appears/enables correctly.

diff --git a/Engine/Source/Editor/SourceControlWindows/Private/SSourceControlChangelists.cpp b/Engine/Source/Editor/SourceControlWindows/Private/SSourceControlChangelists.cpp
--- a/Engine/Source/Editor/SourceControlWindows/Private/SSourceControlChangelists.cpp
+++ b/Engine/Source/Editor/SourceControlWindows/Private/SSourceControlChangelists.cpp
@@ -3039,6 +3039,54 @@ bool SSourceControlChangelistsWidget::CanDiffAgainstWorkspace()
 	return FileListView->GetNumItemsSelected() == 1 && HasShelvedFilesSelected();
 }
 
+void SSourceControlChangelistsWidget::OnSaveAssetsInChangelist()
+{
+	TArray<TSharedPtr<IChangelistTreeItem>> Files = FileListNodes;
+
+	TArray<UPackage*> Packages;
+	Packages.Reserve(Files.Num()); 
+	{
+		for (const TSharedPtr<IChangelistTreeItem>& File : Files)
+		{
+			if (TSharedPtr<IFileViewTreeItem> FileItem = StaticCastSharedPtr<IFileViewTreeItem>(File); FileItem.IsValid())
+			{
+				FString Filename = FileItem->GetFullPathname();
+				FString PackageName = UPackageTools::FilenameToPackageName(Filename);
+
+				if (UPackage* Package = FindPackage(nullptr, *PackageName); Package && Package->IsDirty())
+				{
+					Packages.Add(Package);
+				}
+			}
+		}
+	}
+	const bool bOnlyDirty = false;
+	UEditorLoadingAndSavingUtils::SavePackages(Packages, bOnlyDirty);
+	bShouldRefresh = true;
+}
+
+bool SSourceControlChangelistsWidget::CanSaveAssetsInChangelist()
+{
+	TArray<TSharedPtr<IChangelistTreeItem>> Files = FileListNodes;
+
+	for (const TSharedPtr<IChangelistTreeItem>& File : Files)
+	{
+		if (TSharedPtr<IFileViewTreeItem> FileItem = StaticCastSharedPtr<IFileViewTreeItem>(File); FileItem.IsValid())
+		{
+			FString Filename = FileItem->GetFullPathname();
+			FString PackageName = UPackageTools::FilenameToPackageName(Filename);
+
+			if (UPackage* Package = FindPackage(nullptr, *PackageName); Package && Package->IsDirty())
+			{
+				return true;
+			}
+		}
+	}
+
+
+	return false;
+}
+
 TSharedPtr<SWidget> SSourceControlChangelistsWidget::OnOpenContextMenu()
 {
 	UToolMenus* ToolMenus = UToolMenus::Get();
@@ -3354,6 +3402,19 @@ TSharedPtr<SWidget> SSourceControlChangelistsWidget::OnOpenContextMenu()
 			FUIAction(FExecuteAction::CreateLambda([]() { FUncontrolledChangelistsModule::Get().OnReconcileAssets(); })));
 	}
 
+	if (bHasSelectedChangelist || bHasSelectedUncontrolledChangelist)
+	{
+		Section.AddSeparator("SaveAssets");
+
+		Section.AddMenuEntry(
+			"Save assets",
+			LOCTEXT("SourceControl_SaveAssets", "Save All"),
+			LOCTEXT("SourceControl_SaveAssets_Tooltip", "Save all unsaved assets in the selected changelist."),
+			FSlateIcon(),
+			FUIAction(FExecuteAction::CreateSP(this, &SSourceControlChangelistsWidget::OnSaveAssetsInChangelist), 
+				FCanExecuteAction::CreateSP(this, &SSourceControlChangelistsWidget::CanSaveAssetsInChangelist)));
+	}
+
 	return ToolMenus->GenerateWidget(Menu);
 }
 
diff --git a/Engine/Source/Editor/SourceControlWindows/Private/SSourceControlChangelists.h b/Engine/Source/Editor/SourceControlWindows/Private/SSourceControlChangelists.h
--- a/Engine/Source/Editor/SourceControlWindows/Private/SSourceControlChangelists.h
+++ b/Engine/Source/Editor/SourceControlWindows/Private/SSourceControlChangelists.h
@@ -177,6 +177,10 @@ private:
 	void OnDiffAgainstWorkspace();
 	bool CanDiffAgainstWorkspace();
 
+	/** Save assets in changelist operations */
+	void OnSaveAssetsInChangelist();
+	bool CanSaveAssetsInChangelist();
+
 	/** Uncontrolled changelist module callback. */
 	void OnUncontrolledChangelistStateChanged();
 
