SourceControl: Display a warning if the set of files the user wants to submit contains files that are (exclusively) checked out by others - in this case attempting to submit is pointless.

#rnx

diff --git a/Engine/Source/Editor/SourceControlWindows/Private/SSourceControlSubmit.cpp b/Engine/Source/Editor/SourceControlWindows/Private/SSourceControlSubmit.cpp
--- a/Engine/Source/Editor/SourceControlWindows/Private/SSourceControlSubmit.cpp
+++ b/Engine/Source/Editor/SourceControlWindows/Private/SSourceControlSubmit.cpp
@@ -180,6 +180,17 @@ void SSourceControlSubmitWidget::Construct(const FArguments& InArgs)
 		ListViewItems.Add(MakeShared<FFileTreeItem>(Item));
 	}
 
+	if (CanCheckOut())
+	{
+		for (const auto& Item : InArgs._Items.Get())
+		{
+			if (Item->IsCheckedOutOther() && !Item->CanCheckout()) // Is it an exclusive checkout?
+			{
+				CheckedOutByOther.Add(Item);
+			}
+		}
+	}
+
 	TSharedRef<SHeaderRow> HeaderRowWidget = SNew(SHeaderRow);
 
 	if (bAllowUncheckFiles)
@@ -320,7 +331,7 @@ void SSourceControlSubmitWidget::Construct(const FArguments& InArgs)
 					{
 						if (IsWarningPanelVisible() == EVisibility::Visible)
 						{
-							return ChangeListDescriptionTextCtrl->GetText().IsEmpty() ? EVisibility::Visible : EVisibility::Collapsed;
+							return IsDescriptionWarningVisible() ? EVisibility::Visible : EVisibility::Collapsed;
 						}
 						else
 						{
@@ -343,7 +354,7 @@ void SSourceControlSubmitWidget::Construct(const FArguments& InArgs)
 					{
 						if (IsWarningPanelVisible() == EVisibility::Visible)
 						{
-							return !ChangeListDescriptionTextCtrl->GetText().IsEmpty() ? EVisibility::Visible : EVisibility::Collapsed;
+							return IsValidationWarningVisible() ? EVisibility::Visible : EVisibility::Collapsed;
 						}
 						else
 						{
@@ -354,7 +365,30 @@ void SSourceControlSubmitWidget::Construct(const FArguments& InArgs)
 				.Padding(5)
 				[
 					SNew(SErrorText)
-					.ErrorText(NSLOCTEXT("SourceControl.SubmitPanel", "ChangeListErrorWarning", "Errors need to be fixed to submit")) // Other errors exist and a better mechanism should be built in to display the right error. 
+					.ErrorText(NSLOCTEXT("SourceControl.SubmitPanel", "ChangeListErrorWarning", "Errors need to be fixed to submit"))
+				]
+			]
+			+ SVerticalBox::Slot()
+			.AutoHeight()
+			[
+				SNew(SBorder)
+				.Visibility_Lambda(
+					[this]()
+					{
+						if (IsWarningPanelVisible() == EVisibility::Visible)
+						{
+							return IsCheckoutWarningVisible() ? EVisibility::Visible : EVisibility::Collapsed;
+						}
+						else
+						{
+							return EVisibility::Collapsed;
+						}
+					}
+				)
+				.Padding(5)
+				[
+					SNew(SErrorText)
+					.ErrorText(NSLOCTEXT("SourceControl.SubmitPanel", "ChangeListCheckedOutWarning", "Checkouts by other users need to be cleared to submit"))
 				]
 			]
 		];
@@ -767,6 +801,23 @@ void SSourceControlSubmitWidget::Tick(const FGeometry& AllottedGeometry, const d
 	ISourceControlModule::Get().Tick();
 
 	FTSTicker::GetCoreTicker().Tick(InDeltaTime);
+
+	// The user can release locks in-widget, so need to check for lock status changes.
+	if (CheckedOutByOther.Num() > 0)
+	{
+		for (int32 Index = 0; Index < CheckedOutByOther.Num(); /* */)
+		{
+			const FSourceControlStateRef& State = CheckedOutByOther[Index];
+			if (State->IsCheckedOutOther() && !State->CanCheckout())
+			{
+				++Index;
+			}
+			else
+			{
+				CheckedOutByOther.RemoveAtSwap(Index);
+			}
+		}
+	}
 }
 
 ECheckBoxState SSourceControlSubmitWidget::GetToggleSelectedState() const
@@ -870,7 +921,7 @@ FReply SSourceControlSubmitWidget::SaveAndCloseClicked()
 
 bool SSourceControlSubmitWidget::IsSubmitEnabled() const
 {
-	return bAllowSubmit && !ChangeListDescriptionTextCtrl->GetText().IsEmpty() && ListViewItems.Num() > 0;
+	return bAllowSubmit && !ChangeListDescriptionTextCtrl->GetText().IsEmpty() && ListViewItems.Num() > 0 && CheckedOutByOther.Num() == 0;
 }
 
 
@@ -879,6 +930,29 @@ EVisibility SSourceControlSubmitWidget::IsWarningPanelVisible() const
 	return IsSubmitEnabled() ? EVisibility::Collapsed : EVisibility::Visible;
 }
 
+bool SSourceControlSubmitWidget::IsValidationWarningVisible() const
+{
+	const bool bValidationWarningVisible = !bAllowSubmit;
+	return bValidationWarningVisible;
+}
+
+bool SSourceControlSubmitWidget::IsCheckoutWarningVisible() const
+{
+	const bool bValidationWarningVisible = IsValidationWarningVisible();
+	const bool bCheckoutWarningVisible = CheckedOutByOther.Num() > 0;
+	return bCheckoutWarningVisible && !bValidationWarningVisible;
+}
+
+bool SSourceControlSubmitWidget::IsDescriptionWarningVisible() const
+{
+	const bool bValidationWarningVisible = IsValidationWarningVisible();
+	const bool bCheckoutWarningVisible = IsCheckoutWarningVisible();
+	const bool bDescriptionWarningVisible = ChangeListDescriptionTextCtrl->GetText().IsEmpty();
+	return bDescriptionWarningVisible && !bValidationWarningVisible && !bCheckoutWarningVisible;
+}
+
+
+
 
 void SSourceControlSubmitWidget::OnCheckStateChanged_KeepCheckedOut(ECheckBoxState InState)
 {
diff --git a/Engine/Source/Editor/SourceControlWindows/Private/SSourceControlSubmit.h b/Engine/Source/Editor/SourceControlWindows/Private/SSourceControlSubmit.h
--- a/Engine/Source/Editor/SourceControlWindows/Private/SSourceControlSubmit.h
+++ b/Engine/Source/Editor/SourceControlWindows/Private/SSourceControlSubmit.h
@@ -130,6 +130,11 @@ private:
 	/** Check if the warning panel should be visible. */
 	EVisibility IsWarningPanelVisible() const;
 
+	/** Check if the warning should be visible in the warning panel. */
+	bool IsValidationWarningVisible() const;
+	bool IsCheckoutWarningVisible() const;
+	bool IsDescriptionWarningVisible() const;
+
 	/** Called when the Keep checked out Checkbox is changed */
 	void OnCheckStateChanged_KeepCheckedOut(ECheckBoxState InState);
 
@@ -197,6 +202,9 @@ private:
 	/** Collection of objects (Widgets) to display in the List View. */
 	TArray<FChangelistTreeItemPtr> ListViewItems;
 
+	/** Collection of states that are checked out by other users. */
+	TArray<FSourceControlStateRef> CheckedOutByOther;
+
 	/** Pointer to the parent modal window */
 	TWeakPtr<SWindow> ParentFrame;
 
@@ -212,6 +220,7 @@ private:
 	/** Whether a diff against the depot may be performed from within the submit dialog */
 	bool bAllowDiffAgainstDepot;
 
+	/** Whether the content verse path should be shown or not */
 	bool bShowingContentVersePath;
 
 	/** Specify which column to sort with */
