SourceControl: Only display the 'checked out by others' warning for SCC systems that do not allow individual submit.

This is a follow-up on #46663021.

#rnx

diff --git a/Engine/Source/Editor/SourceControlWindows/Private/SSourceControlSubmit.cpp b/Engine/Source/Editor/SourceControlWindows/Private/SSourceControlSubmit.cpp
--- a/Engine/Source/Editor/SourceControlWindows/Private/SSourceControlSubmit.cpp
+++ b/Engine/Source/Editor/SourceControlWindows/Private/SSourceControlSubmit.cpp
@@ -180,7 +180,7 @@ void SSourceControlSubmitWidget::Construct(const FArguments& InArgs)
 		ListViewItems.Add(MakeShared<FFileTreeItem>(Item));
 	}
 
-	if (CanCheckOut())
+	if (CanCheckOut() && !bAllowUncheckFiles)
 	{
 		for (const auto& Item : InArgs._Items.Get())
 		{
