SourceControl: Fixed invalid 'checkouts by other users need to be cleared' error after reverting files that have locks.

#rnx

diff --git a/Engine/Source/Editor/SourceControlWindows/Private/SSourceControlSubmit.cpp b/Engine/Source/Editor/SourceControlWindows/Private/SSourceControlSubmit.cpp
--- a/Engine/Source/Editor/SourceControlWindows/Private/SSourceControlSubmit.cpp
+++ b/Engine/Source/Editor/SourceControlWindows/Private/SSourceControlSubmit.cpp
@@ -722,15 +722,19 @@ void SSourceControlSubmitWidget::OnRevert()
 		return;
 	}
 
-	auto RemoveItemsFromListView = [this](TArray<FString>& ItemsToRemove)
+	auto RemoveItemsFromListView = [this](const TSet<FString>& ItemsToRemove)
 	{
 		ListViewItems.RemoveAll([&ItemsToRemove](const FChangelistTreeItemPtr& ListViewItem) -> bool
 			{
-				return ItemsToRemove.ContainsByPredicate([&ListViewItem](const FString& ItemToRemove) -> bool
-					{
-						return ItemToRemove == GetFileItem(ListViewItem)->GetFileName().ToString();
-					}
-				);
+				const FString FilePath = GetFileItem(ListViewItem)->GetFileName().ToString();
+				return ItemsToRemove.Contains(FilePath);
+			}
+		);
+
+		CheckedOutByOther.RemoveAll([&ItemsToRemove](const FSourceControlStateRef& StateItem) -> bool
+			{
+				const FString& FilePath = StateItem->GetFilename();
+				return ItemsToRemove.Contains(FilePath);
 			}
 		);
 	};
@@ -760,12 +764,12 @@ void SSourceControlSubmitWidget::OnRevert()
 		if (PackagesToRevert.Num() > 0)
 		{
 			bAnyReverted = SourceControlHelpers::RevertAndReloadPackages(PackagesToRevert, /*bRevertAll=*/false, /*bReloadWorld=*/true);
-			RemoveItemsFromListView(PackagesToRevert);
+			RemoveItemsFromListView(TSet<FString>(PackagesToRevert));
 		}
 		if (FilesToRevert.Num() > 0)
 		{
 			bAnyReverted |= SourceControlHelpers::RevertFiles(FilesToRevert);
-			RemoveItemsFromListView(FilesToRevert);
+			RemoveItemsFromListView(TSet<FString>(FilesToRevert));
 		}
 		
 		if (bAnyReverted)
