Update source control Submit button style.

- Taking the chance now that we are updating styles for UCS-16668

#rnx
#Synced-CL 48809636

diff --git a/Engine/Source/Editor/SourceControlWindows/Private/SSourceControlSubmit.cpp b/Engine/Source/Editor/SourceControlWindows/Private/SSourceControlSubmit.cpp
--- a/Engine/Source/Editor/SourceControlWindows/Private/SSourceControlSubmit.cpp
+++ b/Engine/Source/Editor/SourceControlWindows/Private/SSourceControlSubmit.cpp
@@ -535,6 +535,7 @@ void SSourceControlSubmitWidget::Construct(const FArguments& InArgs)
 			SNew(SButton)
 			.HAlign(HAlign_Center)
 			.ContentPadding(FAppStyle::GetMargin("StandardDialog.ContentPadding"))
+			.ButtonStyle(&FAppStyle::GetWidgetStyle<FButtonStyle>("PrimaryButton"))
 			.IsEnabled(this, &SSourceControlSubmitWidget::IsSubmitEnabled)
 			.Text( NSLOCTEXT("SourceControl.SubmitPanel", "OKButton", "Submit") )
 			.OnClicked(this, &SSourceControlSubmitWidget::SubmitClicked)
