SourceControl: Attach diff icon to 'diff against workspace file' and 'diff against depot' in History dialog.

#rnx

#Synced-CL 50095366

diff --git a/Engine/Source/Editor/SourceControlWindows/Private/SSourceControlHistory.cpp b/Engine/Source/Editor/SourceControlWindows/Private/SSourceControlHistory.cpp
--- a/Engine/Source/Editor/SourceControlWindows/Private/SSourceControlHistory.cpp
+++ b/Engine/Source/Editor/SourceControlWindows/Private/SSourceControlHistory.cpp
@@ -1653,7 +1653,7 @@ void SSourceControlHistoryWidget::CreateDiffMenu(UToolMenu* InToolMenu)
 		TEXT("DiffAgainstPrevious"),
 		NSLOCTEXT("SourceControl.HistoryWindow.Menu", "DiffAgainstPrev", "Diff Against Previous Revision"),
 		NSLOCTEXT("SourceControl.HistoryWindow.Menu", "DiffAgainstPrevTooltip", "See changes between this revision and the previous one."),
-		FSlateIcon(),
+		FSlateIcon(FAppStyle::GetAppStyleSetName(), "SourceControl.Actions.Diff"),
 		FUIAction(
 			FExecuteAction::CreateSP(PinnedHistoryWidget.Get(), &SSourceControlHistoryWidget::OnDiffAgainstPreviousRev),
 			FCanExecuteAction::CreateSP(PinnedHistoryWidget.Get(), &SSourceControlHistoryWidget::CanDiffAgainstPreviousRev)
@@ -1664,7 +1664,7 @@ void SSourceControlHistoryWidget::CreateDiffMenu(UToolMenu* InToolMenu)
 		TEXT("DiffAgainstWorkspace"),
 		NSLOCTEXT("SourceControl.HistoryWindow.Menu", "DiffAgainstWorkspace", "Diff Against Workspace File"),
 		NSLOCTEXT("SourceControl.HistoryWindow.Menu", "DiffAgainstWorkspaceTooltip", "See changes between this revision and your version of the asset."),
-		FSlateIcon(),
+		FSlateIcon(FAppStyle::GetAppStyleSetName(), "SourceControl.Actions.Diff"),
 		FUIAction(
 			FExecuteAction::CreateSP(PinnedHistoryWidget.Get(), &SSourceControlHistoryWidget::OnDiffAgainstWorkspace),
 			FCanExecuteAction::CreateSP(PinnedHistoryWidget.Get(), &SSourceControlHistoryWidget::CanDiffAgainstWorkspace)
