SourceControl: Fixed CHECK-IN CHANGES failing when using a snapshot based SCC provider such as Skein when two users would perform that operation at approximately the same time.

This should be a no-op for other SCC that use this code path such as Git/Subversion as they have 'UsesSnapshots()' returning `false`.

#rnx

#Synced-CL 50143849

diff --git a/Engine/Source/Editor/SourceControlWindows/Private/SourceControlWindows.cpp b/Engine/Source/Editor/SourceControlWindows/Private/SourceControlWindows.cpp
--- a/Engine/Source/Editor/SourceControlWindows/Private/SourceControlWindows.cpp
+++ b/Engine/Source/Editor/SourceControlWindows/Private/SourceControlWindows.cpp
@@ -432,35 +432,5 @@
 	}
 
 	//~~~~~~~~~~~~~~~~~~~~~~~~~~~~~~~~~~~~~~~~~~
-	// Sync to latest if using snapshots.
-	if (ISourceControlModule::Get().GetProvider().UsesSnapshots())
-	{
-		const bool bSyncNeeded = FSourceControlWindows::CanSyncLatest();
-		if (bSyncNeeded)
-		{
-			if (!FSourceControlWindows::SyncLatest())
-			{
-				OutResultInfo.Description = LOCTEXT("SCC_Checkin_Aborted_Sync", "File check in aborted because the sync to the latest snapshot failed.");
-				return false;
-			}
-
-			TArray<FSourceControlStateRef> Conflicts = ISourceControlModule::Get().GetProvider().GetCachedStateByPredicate(
-				[](const FSourceControlStateRef& State)
-				{
-					return State->IsConflicted();
-				}
-			);
-
-			const bool bConflictsRemaining = (Conflicts.Num() > 0);
-			if (bConflictsRemaining)
-			{
-				OutResultInfo.Description = LOCTEXT("SCC_Checkin_Aborted_Conflicts", "File check in aborted because the sync to the latest snapshot resulted in conflicts that need to be resolved.");
-				return false;
-			}
-		}
-	}
-
-
-	//~~~~~~~~~~~~~~~~~~~~~~~~~~~~~~~~~~~~~~~~~~
 	// Get description from the dialog
 	FChangeListDescription Description;
@@ -492,46 +462,44 @@
 		}
 	}
 
-	//~~~~~~~~~~~~~~~~~~~~~~~~~~~~~~~~~~~~~~~~~~
-	// Mark files for add as needed
+	bool bSuccess = true; // Overall success
 	
-	bool bSuccess = true;  // Overall success
-	bool bAddSuccess = true;
-	bool bCheckinSuccess = true;
-	bool bCheckinCancelled = false;
-
-	TArray<FString> CombinedFileList = Description.FilesForAdd;
+	TArray<FString> CombinedFileList;
+	CombinedFileList.Append(Description.FilesForAdd);
 	CombinedFileList.Append(Description.FilesForSubmit);
 
+	//~~~~~~~~~~~~~~~~~~~~~~~~~~~~~~~~~~~~~~~~~~
+	// Any files to check in?
+	if (CombinedFileList.Num() == 0)
+	{
+		OutResultInfo.Result = ECommandResult::Succeeded;
+
+		if (OutResultInfo.Description.IsEmpty())
+		{
+			OutResultInfo.Description = LOCTEXT("SCC_No_Files", "No files were selected to check in to version control.");
+		}
+
+		return true;
+	}
+
+	//~~~~~~~~~~~~~~~~~~~~~~~~~~~~~~~~~~~~~~~~~~
+	// Mark files for add as needed
 	if (Description.FilesForAdd.Num() > 0)
 	{
-		bAddSuccess = SourceControlProvider.Execute(ISourceControlOperation::Create<FMarkForAdd>(), Description.FilesForAdd) == ECommandResult::Succeeded;
-		bSuccess &= bAddSuccess;
-
 		OutResultInfo.FilesAdded = Description.FilesForAdd;
 
+		const bool bAddSuccess = SourceControlProvider.Execute(ISourceControlOperation::Create<FMarkForAdd>(), Description.FilesForAdd) == ECommandResult::Succeeded;
 		if (!bAddSuccess)
 		{
 			// Note that this message may be overwritten with a checkin error below.
 			OutResultInfo.Description = LOCTEXT("SCC_Add_Files_Error", "One or more files were not able to be marked for add to version control!");
 		}
+		
+		bSuccess &= bAddSuccess;
 	}
 
 	//~~~~~~~~~~~~~~~~~~~~~~~~~~~~~~~~~~~~~~~~~~
-	// Any files to check in?
-	if (CombinedFileList.Num() == 0)
-	{
-		OutResultInfo.Result = bSuccess ? ECommandResult::Succeeded : ECommandResult::Failed;
-
-		if (OutResultInfo.Description.IsEmpty())
-		{
-			OutResultInfo.Description = LOCTEXT("SCC_No_Files", "No files were selected to check in to version control.");
-		}
-
-		return bSuccess;
-	}
-
-	// first check if there is a submit override bound
+	// First check if there is a submit override bound
 	if (ISourceControlWindowsModule::Get().SubmitOverrideDelegate.IsBound())
 	{
 		SSubmitOverrideParameters SubmitOverrideParameters;
@@ -565,6 +533,8 @@
 		}
 	}
 
+	//~~~~~~~~~~~~~~~~~~~~~~~~~~~~~~~~~~~~~~~~~~
+	// Virtualize the files that need to be submitted
 	FText VirtualizationFailureMsg;
 	if (!TryToVirtualizeFilesToSubmit(CombinedFileList, Description.Description, VirtualizationFailureMsg))
 	{
@@ -576,31 +546,76 @@
 		return false;
 	}
 
-	// Check in files
-	TSharedRef<FCheckIn, ESPMode::ThreadSafe> CheckInOperation = ISourceControlOperation::Create<FCheckIn>();
-	CheckInOperation->SetDescription(Description.Description);
-	CheckInOperation->SetKeepCheckedOut(SourceControlWidget->WantToKeepCheckedOut());
+	bool bCheckinSuccess = true;
+	bool bCheckinCancelled = false;
+	bool bSyncLatestNeeded = bUsesSnapshots ? FSourceControlWindows::CanSyncLatest() : false;
 
-	ECommandResult::Type CheckInResult = SourceControlProvider.Execute(CheckInOperation, CombinedFileList);
-	bCheckinSuccess = CheckInResult == ECommandResult::Succeeded;
-	bCheckinCancelled = CheckInResult == ECommandResult::Cancelled;
+	FText CheckinSuccessMessage;
 
+	do
+	{
+		//~~~~~~~~~~~~~~~~~~~~~~~~~~~~~~~~~~~~~~~~~~
+		// Sync to latest if using snapshots as we cannot check in unless we're at head
+		if (bSyncLatestNeeded)
+		{
+			if (!FSourceControlWindows::SyncLatest())
+			{
+				OutResultInfo.Description = LOCTEXT("SCC_Checkin_Aborted_Sync", "File check in aborted because the sync to the latest snapshot failed.");
+				return false;
+			}
+
+			TArray<FSourceControlStateRef> Conflicts = ISourceControlModule::Get().GetProvider().GetCachedStateByPredicate(
+				[](const FSourceControlStateRef& State)
+				{
+					return State->IsConflicted();
+				}
+			);
+
+			const bool bConflictsRemaining = (Conflicts.Num() > 0);
+			if (bConflictsRemaining)
+			{
+				OutResultInfo.Description = LOCTEXT("SCC_Checkin_Aborted_Conflicts", "File check in aborted because the sync to the latest snapshot resulted in conflicts that need to be resolved.");
+				return false;
+			}
+		}
+
+		//~~~~~~~~~~~~~~~~~~~~~~~~~~~~~~~~~~~~~~~~~~
+		// Check in files
+		TSharedRef<FCheckIn, ESPMode::ThreadSafe> CheckInOperation = ISourceControlOperation::Create<FCheckIn>();
+		CheckInOperation->SetDescription(Description.Description);
+		CheckInOperation->SetKeepCheckedOut(SourceControlWidget->WantToKeepCheckedOut());
+
+		ECommandResult::Type CheckInResult = SourceControlProvider.Execute(CheckInOperation, CombinedFileList);
+		bCheckinSuccess = CheckInResult == ECommandResult::Succeeded;
+		bCheckinCancelled = CheckInResult == ECommandResult::Cancelled;
+
+		//~~~~~~~~~~~~~~~~~~~~~~~~~~~~~~~~~~~~~~~~~~
+		// Check if another user beat us to snapshot creation
+		// This can happen - for example - when taking a long time writing a description
+		bSyncLatestNeeded = bUsesSnapshots ? FSourceControlWindows::CanSyncLatest() : false;
+
+		//~~~~~~~~~~~~~~~~~~~~~~~~~~~~~~~~~~~~~~~~~~
+		// Store the success message for later use
+		CheckinSuccessMessage = CheckInOperation->GetSuccessMessage();
+
+	} while ((bUsesSnapshots && bSyncLatestNeeded) && !bCheckinSuccess && !bCheckinCancelled);
+
 	bSuccess &= bCheckinSuccess;
 
 	if (bCheckinSuccess)
 	{
 		// report success with a notification
-		FNotificationInfo Info(CheckInOperation->GetSuccessMessage());
+		FNotificationInfo Info(CheckinSuccessMessage);
 		Info.ExpireDuration = 8.0f;
 		Info.HyperlinkText = LOCTEXT("SCC_Checkin_ShowLog", "Show Message Log");
 		Info.Hyperlink = FSimpleDelegate::CreateStatic([](){ FMessageLog("SourceControl").Open(EMessageSeverity::Info, true); });
 		FSlateNotificationManager::Get().AddNotification(Info);
 
 		// also add to the log
-		FMessageLog("SourceControl").Info(CheckInOperation->GetSuccessMessage());
+		FMessageLog("SourceControl").Info(CheckinSuccessMessage);
 
 		OutResultInfo.Result         = ECommandResult::Succeeded;
-		OutResultInfo.Description    = CheckInOperation->GetSuccessMessage();
+		OutResultInfo.Description    = CheckinSuccessMessage;
 		OutResultInfo.FilesSubmitted = Description.FilesForSubmit;
 	}
 
