Fix an issue in View Changes where saving a dirty asset from an uncontrolled changelist would move it to the default controlled changelist and mark-for-add it there, when connected to a source control provider.

Now, it will be saved back to its uncontrolled changelist, which matches the behavior that happens when it is saved from its asset editor.  The issue was caused by the behavior of UEditorLoadingAndSavingUtils::SavePackages as compared to FEditorFileUtils::PromptForCheckoutAndSave.  Using the latter fixes this issue.

In the "Save To Changelist" context menu action, we save before moving the files so additional support fixes are needed now that the dirty assets from uncontrolled changelists are properly saved in their uncontrolled changelists prior to the move.  In SSourceControlChangelistsWidget::GetSelectedFiles, UnsavedAssets always have the IChangelistTreeItem::OfflineFile type so logic was added to check if each filename is contained in an uncontrolled changelist so that they can be properly sorted into the OutUncontrolledFiles and OutControlledFiles buckets for the logic in TryMoveFiles() to handle accordingly.

#Synced-CL 50412461

diff --git a/Engine/Source/Editor/SourceControlWindows/Private/SSourceControlChangelists.cpp b/Engine/Source/Editor/SourceControlWindows/Private/SSourceControlChangelists.cpp
--- a/Engine/Source/Editor/SourceControlWindows/Private/SSourceControlChangelists.cpp
+++ b/Engine/Source/Editor/SourceControlWindows/Private/SSourceControlChangelists.cpp
@@ -99,6 +99,26 @@ bool AreWorldReloadsAllowed()
 	}
 }
 
+bool AnyUncontrolledChangelistContainsFilename(const FString& Filename)
+{
+	FUncontrolledChangelistsModule& UncontrolledChangelistsModule = FUncontrolledChangelistsModule::Get();
+
+	if (UncontrolledChangelistsModule.IsEnabled())
+	{
+		TArray<TSharedRef<FUncontrolledChangelistState>> UncontrolledChangelistStates = UncontrolledChangelistsModule.GetChangelistStates();
+
+		for (const TSharedRef<FUncontrolledChangelistState>& UncontrolledChangelistState : UncontrolledChangelistStates)
+		{
+			if (UncontrolledChangelistState->ContainsFilename(Filename))
+			{
+				return true;
+			}
+		}
+	}
+
+	return false;
+}
+
 /**
  * Returns a new changelist description if needed, appending validation tag.
  * 
@@ -1795,7 +1815,15 @@ void SSourceControlChangelistsWidget::GetSelectedFiles(TArray<FString>& OutContr
 				else if (Parent->GetTreeItemType() == IChangelistTreeItem::UnsavedAssets)
 				{
 					const FString& Filename = StaticCastSharedPtr<FOfflineFileTreeItem>(Item)->GetFullPathname();
-					OutControlledFiles.Add(Filename);
+
+					if (AnyUncontrolledChangelistContainsFilename(Filename))
+					{
+						OutUncontrolledFiles.Add(Filename);
+					}
+					else
+					{
+						OutControlledFiles.Add(Filename);
+					}
 				}
 			}
 		}
@@ -2913,10 +2941,17 @@ TValueOrError<void, void> SSourceControlChangelistsWidget::TryMoveFiles()
 			{
 				UncontrolledChangelistsModule.MoveFilesToUncontrolledChangelist(SelectedControlledFileStates, SelectedUnControlledFileStates, NewUncontrolledChangelist.GetValue());
 			}
-			// SelectedControlledFiles can contain UnsavedAssets when offline even if both sets of states are empty.
-			else if (!SelectedControlledFiles.IsEmpty())
+			// SelectedControlledFiles and/or SelectedUncontrolledFiles can contain UnsavedAssets (IChangelistTreeItem::OfflineFile)
+			else
 			{
-				UncontrolledChangelistsModule.MoveFilesToUncontrolledChangelist(SelectedControlledFiles, NewUncontrolledChangelist.GetValue());
+				if (!SelectedControlledFiles.IsEmpty())
+				{
+					UncontrolledChangelistsModule.MoveFilesToUncontrolledChangelist(SelectedControlledFiles, NewUncontrolledChangelist.GetValue());
+				}
+				if (!SelectedUncontrolledFiles.IsEmpty())
+				{
+					UncontrolledChangelistsModule.MoveFilesToUncontrolledChangelist(SelectedUncontrolledFiles, NewUncontrolledChangelist.GetValue());
+				}
 			}
 		});
 	}
@@ -2962,7 +2997,7 @@ TValueOrError<void, void> SSourceControlChangelistsWidget::TryMoveFiles()
 						}
 					}));
 			}
-			else if (!SelectedUncontrolledFiles.IsEmpty())
+			if (!SelectedUncontrolledFiles.IsEmpty())
 			{
 				ExecuteUncontrolledChangelistOperation(LOCTEXT("Moving_Uncontrolled_Files_To_Changelist", "Moving uncontrolled files..."), [&]()
 				{
@@ -2985,10 +3020,17 @@ TValueOrError<void, void> SSourceControlChangelistsWidget::TryMoveFiles()
 				{
 					FUncontrolledChangelistsModule::Get().MoveFilesToUncontrolledChangelist(SelectedControlledFileStates, SelectedUnControlledFileStates, UncontrolledChangelist);
 				}
-				// SelectedControlledFiles can contain UnsavedAssets when offline even if both sets of states are empty.
-				else if (!SelectedControlledFiles.IsEmpty())
+				// SelectedControlledFiles and/or SelectedUncontrolledFiles can contain UnsavedAssets (IChangelistTreeItem::OfflineFile)
+				else
 				{
-					FUncontrolledChangelistsModule::Get().MoveFilesToUncontrolledChangelist(SelectedControlledFiles, UncontrolledChangelist);
+					if (!SelectedControlledFiles.IsEmpty())
+					{
+						FUncontrolledChangelistsModule::Get().MoveFilesToUncontrolledChangelist(SelectedControlledFiles, UncontrolledChangelist);
+					}
+					if (!SelectedUncontrolledFiles.IsEmpty())
+					{
+						FUncontrolledChangelistsModule::Get().MoveFilesToUncontrolledChangelist(SelectedUncontrolledFiles, UncontrolledChangelist);
+					}
 				}
 			});
 		}
@@ -3059,7 +3101,9 @@ void SSourceControlChangelistsWidget::OnSaveAssetsInChangelist()
 		}
 	}
 	const bool bOnlyDirty = false;
-	UEditorLoadingAndSavingUtils::SavePackages(Packages, bOnlyDirty);
+	const bool bPromptToSave = false;
+	FEditorFileUtils::PromptForCheckoutAndSave(Packages, bOnlyDirty, bPromptToSave);
+
 	bShouldRefresh = true;
 }
 
@@ -3153,7 +3197,8 @@ TSharedPtr<SWidget> SSourceControlChangelistsWidget::OnOpenContextMenu()
 					}
 				}
 				const bool bOnlyDirty = false;
-				UEditorLoadingAndSavingUtils::SavePackages(Packages, bOnlyDirty);
+				const bool bPromptToSave = false;
+				FEditorFileUtils::PromptForCheckoutAndSave(Packages, bOnlyDirty, bPromptToSave);
 
 				// Request a changelists refresh in the next tick
 				bShouldRefresh = true;
@@ -3211,7 +3256,8 @@ TSharedPtr<SWidget> SSourceControlChangelistsWidget::OnOpenContextMenu()
 					}
 				}
 				const bool bOnlyDirty = false;
-				UEditorLoadingAndSavingUtils::SavePackages(Packages, bOnlyDirty);
+				const bool bPromptToSave = false;
+				FEditorFileUtils::PromptForCheckoutAndSave(Packages, bOnlyDirty, bPromptToSave);
 
 				// Move the files only after saving the packages.
 				// If we are moving files to an uncontrolled changelist, they will be reverted from source control during the move.
