SourceControl: Made some arguments to `PromptForCheckedOut` a `const` where they should be `const`.

#rnx

#Synced-CL 50095366

diff --git a/Engine/Source/Editor/SourceControlWindows/Private/SourceControlWindows.cpp b/Engine/Source/Editor/SourceControlWindows/Private/SourceControlWindows.cpp
--- a/Engine/Source/Editor/SourceControlWindows/Private/SourceControlWindows.cpp
+++ b/Engine/Source/Editor/SourceControlWindows/Private/SourceControlWindows.cpp
@@ -657,7 +657,7 @@ bool FSourceControlWindows::PromptForCheckin(bool bUseSourceControlStateCache, c
 	return PromptForCheckin(ResultInfo, InPackageNames, InPendingDeletePaths, InConfigFiles, bUseSourceControlStateCache);
 }
 
-bool FSourceControlWindows::PromptForCheckedOut(bool bUseSourceControlStateCache, TArray<FString>& InFileNames, FCheckedOutSetupInfo& InSetupInfo)
+bool FSourceControlWindows::PromptForCheckedOut(bool bUseSourceControlStateCache, const TArray<FString>& InFileNames, const FCheckedOutSetupInfo& InSetupInfo)
 {
 	TArray<FSourceControlStateRef> Items;
 
diff --git a/Engine/Source/Editor/SourceControlWindows/Public/SourceControlWindows.h b/Engine/Source/Editor/SourceControlWindows/Public/SourceControlWindows.h
--- a/Engine/Source/Editor/SourceControlWindows/Public/SourceControlWindows.h
+++ b/Engine/Source/Editor/SourceControlWindows/Public/SourceControlWindows.h
@@ -166,7 +166,7 @@ public:
 		bool bShowColumnUserName = true;
 	};
 
-	static UE_API bool PromptForCheckedOut(bool bUseSourceControlStateCache, TArray<FString>& InFileNames, FCheckedOutSetupInfo& InSetupInfo);
+	static UE_API bool PromptForCheckedOut(bool bUseSourceControlStateCache, const TArray<FString>& InFileNames, const FCheckedOutSetupInfo& InSetupInfo);
 
 	/**
 	 * Display file revision history for the provided packages
