SourceControl: Add ability to override the title of the 'Confirm Submit' dialog.

#rnx

#Synced-CL 50815533

#Synced-CL 50815533

diff --git a/Engine/Source/Editor/SourceControlWindows/Private/SSourceControlChangelists.cpp b/Engine/Source/Editor/SourceControlWindows/Private/SSourceControlChangelists.cpp
--- a/Engine/Source/Editor/SourceControlWindows/Private/SSourceControlChangelists.cpp
+++ b/Engine/Source/Editor/SourceControlWindows/Private/SSourceControlChangelists.cpp
@@ -2535,10 +2535,18 @@ void SSourceControlChangelistsWidget::OnSubmitChangelist()
 	// The description once edited by the user in the Submit window.
 	FText UserEditChangelistDescription = ChangelistDescriptionToSubmit;
 
+	// The title from the window
+	FSubmitTitleOverrideParameters SubmitTitleOverrideParameters;
+	SubmitTitleOverrideParameters.ToSubmit.SetSubtype<FString>(ChangelistState->GetChangelist()->GetIdentifier());
+
+	const FText Title = ISourceControlWindowsModule::Get().SubmitTitleOverrideDelegate.IsBound()
+		? ISourceControlWindowsModule::Get().SubmitTitleOverrideDelegate.Execute(SubmitTitleOverrideParameters)
+		: NSLOCTEXT("SourceControl.ConfirmSubmit", "Title", "Confirm Submit");
+
 	TSharedRef<SWindow> NewWindow = SNew(SWindow)
-		.Title(NSLOCTEXT("SourceControl.ConfirmSubmit", "Title", "Confirm Submit"))
+		.Title(Title)
 		.SizingRule(ESizingRule::UserSized)
-		.ClientSize(FVector2D(800, 600))
+		.ClientSize(FVector2D(800, 700))
 		.SupportsMaximize(true)
 		.SupportsMinimize(false);
 
diff --git a/Engine/Source/Editor/SourceControlWindows/Private/SourceControlWindows.cpp b/Engine/Source/Editor/SourceControlWindows/Private/SourceControlWindows.cpp
--- a/Engine/Source/Editor/SourceControlWindows/Private/SourceControlWindows.cpp
+++ b/Engine/Source/Editor/SourceControlWindows/Private/SourceControlWindows.cpp
@@ -385,12 +385,21 @@ bool FSourceControlWindows::PromptForCheckin(FCheckinResultInfo& OutResultInfo,
 		? GetFilesValidationResult(States, ChangelistValidationTitle, ChangelistValidationWarningsText, ChangelistValidationErrorsText)
 		: true;
 
+	//~~~~~~~~~~~~~~~~~~~~~~~~~~~~~~~~~~~~~~~~~~
+	// Generate the submit files window title
+	FSubmitTitleOverrideParameters SubmitTitleOverrideParameters;
+	SubmitTitleOverrideParameters.ToSubmit.SetSubtype<TArray<FString>>(AllFiles);
+
+	const FText Title = ISourceControlWindowsModule::Get().SubmitTitleOverrideDelegate.IsBound()
+		? ISourceControlWindowsModule::Get().SubmitTitleOverrideDelegate.Execute(SubmitTitleOverrideParameters)
+		: NSLOCTEXT("SourceControl.ConfirmSubmit", "Title", "Confirm Submit");
+
 	//~~~~~~~~~~~~~~~~~~~~~~~~~~~~~~~~~~~~~~~~~~
 	// Create a submit files window
 	TSharedRef<SWindow> NewWindow = SNew(SWindow)
-		.Title(NSLOCTEXT("SourceControl.ConfirmSubmit", "Title", "Confirm Submit"))
+		.Title(Title)
 		.SizingRule(ESizingRule::UserSized)
-		.ClientSize(FVector2D(800, 600))
+		.ClientSize(FVector2D(800, 700))
 		.SupportsMaximize(true)
 		.SupportsMinimize(false);
 
diff --git a/Engine/Source/Editor/SourceControlWindows/Public/ISourceControlWindowsModule.h b/Engine/Source/Editor/SourceControlWindows/Public/ISourceControlWindowsModule.h
--- a/Engine/Source/Editor/SourceControlWindows/Public/ISourceControlWindowsModule.h
+++ b/Engine/Source/Editor/SourceControlWindows/Public/ISourceControlWindowsModule.h
@@ -26,8 +26,14 @@ struct SSubmitOverrideParameters
 	FString Description;
 };
 
-DECLARE_DELEGATE_RetVal_TwoParams(FSubmitOverrideReply, FCanSubmitOverrideDelegate, const SSubmitOverrideParameters& InParameters /*InParameters*/, FText* ErrorMessageOut);
-DECLARE_DELEGATE_RetVal_OneParam(FSubmitOverrideReply, FSubmitOverrideDelegate, const SSubmitOverrideParameters& InParameters /*InParameters*/);
+struct FSubmitTitleOverrideParameters
+{
+	TUnion<FString, TArray<FString>> ToSubmit;
+};
+
+DECLARE_DELEGATE_RetVal_TwoParams(FSubmitOverrideReply, FCanSubmitOverrideDelegate, const SSubmitOverrideParameters& /*InParameters*/, FText* /*ErrorMessageOut*/);
+DECLARE_DELEGATE_RetVal_OneParam(FSubmitOverrideReply, FSubmitOverrideDelegate, const SSubmitOverrideParameters& /*InParameters*/);
+DECLARE_DELEGATE_RetVal_OneParam(FText, FSubmitTitleOverrideDelegate, const FSubmitTitleOverrideParameters& /*InParameters*/);
 
 class ISourceControlWindowsModule : public IModuleInterface
 {
@@ -65,4 +71,6 @@ public:
 
 	FCanSubmitOverrideDelegate CanSubmitOverrideDelegate;
 	FSubmitOverrideDelegate SubmitOverrideDelegate;
+
+	FSubmitTitleOverrideDelegate SubmitTitleOverrideDelegate;
 };
