SourceControl: Reduce area occupied by the description textbox a little in the 'Confirm Submit' dialog as the majority of people just enter a single line - having some more room for the ListView thus makes sense.

#rnx

#Synced-CL 50815533

diff --git a/Engine/Source/Editor/SourceControlWindows/Private/SSourceControlSubmit.cpp b/Engine/Source/Editor/SourceControlWindows/Private/SSourceControlSubmit.cpp
--- a/Engine/Source/Editor/SourceControlWindows/Private/SSourceControlSubmit.cpp
+++ b/Engine/Source/Editor/SourceControlWindows/Private/SSourceControlSubmit.cpp
@@ -285,22 +285,19 @@ void SSourceControlSubmitWidget::Construct(const FArguments& InArgs)
 	];
 
 	Contents->AddSlot()
-	.FillHeight(1.f)
+	.FillHeight(0.4f)
 	.Padding(FMargin(5, 0, 5, 5))
 	[
-		SNew(SBox)
-		.WidthOverride(520)
-		[
-			SAssignNew(ChangeListDescriptionTextCtrl, SMultiLineEditableTextBox)
-			.SelectAllTextWhenFocused(!bDescriptionIsReadOnly)
-			.Text(GSavedChangeListDescription)
-			.AutoWrapText(true)
-			.IsReadOnly(bDescriptionIsReadOnly)
-		]
+		SAssignNew(ChangeListDescriptionTextCtrl, SMultiLineEditableTextBox)
+		.SelectAllTextWhenFocused(!bDescriptionIsReadOnly)
+		.Text(GSavedChangeListDescription)
+		.AutoWrapText(true)
+		.IsReadOnly(bDescriptionIsReadOnly)
 	];
 
 	Contents->AddSlot()
 	.Padding(FMargin(5, 0))
+	.FillHeight(1.f)
 	[
 		SNew(SBorder)
 		[
@@ -446,6 +443,7 @@ void SSourceControlSubmitWidget::Construct(const FArguments& InArgs)
 			if (!ChangelistResultErrorsText.IsEmpty())
 			{
 				Contents->AddSlot()
+				.AutoHeight()
 				.Padding(FMargin(5))
 				[
 					SNew(SHorizontalBox)
@@ -469,6 +467,7 @@ void SSourceControlSubmitWidget::Construct(const FArguments& InArgs)
 			if (!ChangelistResultWarningsText.IsEmpty())
 			{
 				Contents->AddSlot()
+				.AutoHeight()
 				.Padding(FMargin(5))
 				[
 					SNew(SHorizontalBox)
