SourceControl: Fixed 'description' field not being visible in CONFIRM SUBMIT dialog if there are a lot of validation errors/warnings.

#rnx
#Synced-CL 51451321

#Synced-CL 51451321

diff --git a/Engine/Source/Editor/SourceControlWindows/Private/SSourceControlSubmit.cpp b/Engine/Source/Editor/SourceControlWindows/Private/SSourceControlSubmit.cpp
--- a/Engine/Source/Editor/SourceControlWindows/Private/SSourceControlSubmit.cpp
+++ b/Engine/Source/Editor/SourceControlWindows/Private/SSourceControlSubmit.cpp
@@ -443,7 +443,7 @@ void SSourceControlSubmitWidget::Construct(const FArguments& InArgs)
 			if (!ChangelistResultErrorsText.IsEmpty())
 			{
 				Contents->AddSlot()
-				.AutoHeight()
+				.FillHeight(1.f)
 				.Padding(FMargin(5))
 				[
 					SNew(SHorizontalBox)
@@ -467,7 +467,7 @@ void SSourceControlSubmitWidget::Construct(const FArguments& InArgs)
 			if (!ChangelistResultWarningsText.IsEmpty())
 			{
 				Contents->AddSlot()
-				.AutoHeight()
+				.FillHeight(1.f)
 				.Padding(FMargin(5))
 				[
 					SNew(SHorizontalBox)
