SourceControl: Fixed commit message being 'empty' by containing solely whitespace.

#rnx

#Synced-CL 52128263

diff --git a/Engine/Source/Editor/SourceControlWindows/Private/SSourceControlSubmit.cpp b/Engine/Source/Editor/SourceControlWindows/Private/SSourceControlSubmit.cpp
--- a/Engine/Source/Editor/SourceControlWindows/Private/SSourceControlSubmit.cpp
+++ b/Engine/Source/Editor/SourceControlWindows/Private/SSourceControlSubmit.cpp
@@ -919,7 +919,7 @@ FReply SSourceControlSubmitWidget::SaveAndCloseClicked()
 
 bool SSourceControlSubmitWidget::IsSubmitEnabled() const
 {
-	return bAllowSubmit && !ChangeListDescriptionTextCtrl->GetText().IsEmpty() && ListViewItems.Num() > 0 && CheckedOutByOther.Num() == 0;
+	return bAllowSubmit && !ChangeListDescriptionTextCtrl->GetText().ToString().TrimStartAndEnd().IsEmpty() && ListViewItems.Num() > 0 && CheckedOutByOther.Num() == 0;
 }
 
 
@@ -945,7 +945,7 @@ bool SSourceControlSubmitWidget::IsDescriptionWarningVisible() const
 {
 	const bool bValidationWarningVisible = IsValidationWarningVisible();
 	const bool bCheckoutWarningVisible = IsCheckoutWarningVisible();
-	const bool bDescriptionWarningVisible = ChangeListDescriptionTextCtrl->GetText().IsEmpty();
+	const bool bDescriptionWarningVisible = ChangeListDescriptionTextCtrl->GetText().ToString().TrimStartAndEnd().IsEmpty();
 	return bDescriptionWarningVisible && !bValidationWarningVisible && !bCheckoutWarningVisible;
 }
 
